Removed Mobile FSR from Source/Shaders
Added Mobile FSR as a plugin

diff --git a/Engine/Plugins/Runtime/MobileFSR/MobileFSR.uplugin b/Engine/Plugins/Runtime/MobileFSR/MobileFSR.uplugin
new file mode 100644
--- /dev/null
+++ b/Engine/Plugins/Runtime/MobileFSR/MobileFSR.uplugin
@@ -0,0 +1,24 @@
+{
+	"FileVersion": 1,
+	"Version": 1,
+	"VersionName": "1.0",
+	"FriendlyName": "Mobile FSR",
+	"Description": "Mobile FSR 1.0",
+	"Category": "Rendering",
+	"CreatedBy": "Epic Games, Inc.",
+	"CreatedByURL": "http://epicgames.com",
+	"DocsURL": "",
+	"MarketplaceURL": "",
+	"SupportURL": "",
+	"EnabledByDefault": false,
+	"CanContainContent": true,
+	"IsBetaVersion": false,
+	"Installed": false,
+	"Modules": [
+		{
+			"Name": "MobileFSR",
+			"Type": "Runtime",
+			"LoadingPhase": "PostConfigInit"
+		}
+	]
+}
\ No newline at end of file
diff --git a/Engine/Plugins/Runtime/MobileFSR/Shaders/Private/PostProcessMobileFFX_Common.ush b/Engine/Plugins/Runtime/MobileFSR/Shaders/Private/PostProcessMobileFFX_Common.ush
new file mode 100644
--- /dev/null
+++ b/Engine/Plugins/Runtime/MobileFSR/Shaders/Private/PostProcessMobileFFX_Common.ush
@@ -0,0 +1,154 @@
+// Copyright Epic Games, Inc. All Rights Reserved.
+//------------------------------------------------------------------------------
+// Copyright (c) 2021 Advanced Micro Devices, Inc. All rights reserved.
+// Permission is hereby granted, free of charge, to any person obtaining a copy
+// of this software and associated documentation files(the "Software"), to deal
+// in the Software without restriction, including without limitation the rights
+// to use, copy, modify, merge, publish, distribute, sublicense, and / or sell
+// copies of the Software, and to permit persons to whom the Software is
+// furnished to do so, subject to the following conditions :
+// The above copyright notice and this permission notice shall be included in
+// all copies or substantial portions of the Software.
+// THE SOFTWARE IS PROVIDED "AS IS", WITHOUT WARRANTY OF ANY KIND, EXPRESS OR
+// IMPLIED, INCLUDING BUT NOT LIMITED TO THE WARRANTIES OF MERCHANTABILITY,
+// FITNESS FOR A PARTICULAR PURPOSE AND NONINFRINGEMENT.IN NO EVENT SHALL THE
+// AUTHORS OR COPYRIGHT HOLDERS BE LIABLE FOR ANY CLAIM, DAMAGES OR OTHER
+// LIABILITY, WHETHER IN AN ACTION OF CONTRACT, TORT OR OTHERWISE, ARISING FROM,
+// OUT OF OR IN CONNECTION WITH THE SOFTWARE OR THE USE OR OTHER DEALINGS IN
+// THE SOFTWARE.
+//------------------------------------------------------------------------------
+
+// !!!
+// FidelityFX Super Resolution - Common
+// This header should be included after ffx_a.ush
+// !!!
+
+// =====================================================================================
+// 
+// EOTF Conversions
+// 
+// =====================================================================================
+// FSR output device definitions : matching EFSR_OutputDevice enum in PostProcessFSR.cpp
+#define FSR_SRGB   0
+#define FSR_LINEAR 1
+#define FSR_ST2084 2
+#define FSR_scRGB  3
+
+#define MAX_ODT_NITS_ENUM 0
+#define CONVERT_TO_OUTPUT_DEVICE 0
+
+#include "/Engine/Private/ACES.ush"
+#include "/Engine/Private/GammaCorrectionCommon.ush"
+
+float4 LinearToGamma2(float4 Color) { return sqrt(Color); }
+float3 LinearToGamma2(float3 Color) { return sqrt(Color); }
+float2 LinearToGamma2(float2 Color) { return sqrt(Color); }
+float  LinearToGamma2(float  Color) { return sqrt(Color); }
+
+float4 Gamma2ToLinear(float4 Color) { return Color * Color; }
+float3 Gamma2ToLinear(float3 Color) { return Color * Color; }
+float2 Gamma2ToLinear(float2 Color) { return Color * Color; }
+float  Gamma2ToLinear(float  Color) { return Color * Color; }
+
+#if ENABLE_FP16
+AH4 LinearToGamma2(AH4 Color) { return sqrt(Color); }
+AH3 LinearToGamma2(AH3 Color) { return sqrt(Color); }
+AH2 LinearToGamma2(AH2 Color) { return sqrt(Color); }
+AH1 LinearToGamma2(AH1 Color) { return sqrt(Color); }
+
+AH4 Gamma2ToLinear(AH4 Color) { return Color * Color; }
+AH3 Gamma2ToLinear(AH3 Color) { return Color * Color; }
+AH2 Gamma2ToLinear(AH2 Color) { return Color * Color; }
+AH1 Gamma2ToLinear(AH1 Color) { return Color * Color; }
+#endif
+
+//-----------------------------------------------------------------
+
+// only float3 ST2084ToLinear() is defined in GammaCorrectionCommon.ush
+float ST2084ToLinear_F(float pq)
+{
+	const float m1 = 0.1593017578125; // = 2610. / 4096. * .25;
+	const float m2 = 78.84375;        // = 2523. / 4096. *  128;
+	const float c1 = 0.8359375;       // = 2392. / 4096. * 32 - 2413./4096.*32 + 1;
+	const float c2 = 18.8515625;      // = 2413. / 4096. * 32;
+	const float c3 = 18.6875;         // = 2392. / 4096. * 32;
+	const float C = 10000.;
+
+	float Np = pow(pq, 1. / m2);
+	float L = Np - c1;
+	L = max(0., L);
+	L = L / (c2 - c3 * Np);
+	L = pow(L, 1. / m1);
+	float P = L * C;
+
+	return P;
+}
+float  ST2084ToGamma2(float  pq) { return LinearToGamma2(ST2084ToLinear_F(pq)); }
+float2 ST2084ToGamma2(float2 pq) { return LinearToGamma2(float2(ST2084ToLinear_F(pq.r), ST2084ToLinear_F(pq.g))); }
+float3 ST2084ToGamma2(float3 pq) { return LinearToGamma2(ST2084ToLinear(pq)); }
+float4 ST2084ToGamma2(float4 pq) { return LinearToGamma2(float4(ST2084ToLinear(pq.rgb), ST2084ToLinear_F(pq.a))); }
+
+//-----------------------------------------------------------------
+
+// Either 1000 or 2000 depending on UE4 Tonemapper enum, 
+// should be provided by the engine
+#if MAX_ODT_NITS_ENUM == 1
+#define HDR_MAX_ODT_NITS 2000
+#else 
+#define HDR_MAX_ODT_NITS 1000
+#endif
+
+#define HDR_MAX_NITS     10000
+#define HDR_MAX_NITS_INV 0.0001f
+
+// matches white point in ST2084ToScRGB()
+#define HDR_WHITE_POINT 80.0f
+
+float3 Gamma2ToScRGB(float3 Color)
+{
+	const float3x3 Rec2020_2_sRGB = mul(XYZ_2_sRGB_MAT, Rec2020_2_XYZ_MAT);
+	const float MaxODTNits = HDR_MAX_ODT_NITS;
+	const float scRGBScale = MaxODTNits / HDR_WHITE_POINT;
+
+	// Gamma2 -> Linear encoding
+	float3 LinearColor_Rec2020 = Gamma2ToLinear(Color);
+
+	// Rec2020 -> sRGB primaries
+	float3 LinearColor_sRGB = mul(Rec2020_2_sRGB, LinearColor_Rec2020);
+
+	// Linear_sRGB -> scRGB
+	return LinearColor_sRGB * scRGBScale;
+}
+
+float3 scRGBToGamma2(float3 Color_scRGB)
+{
+	const float3x3 sRGB_2_Rec2020 = mul(XYZ_2_Rec2020_MAT, sRGB_2_XYZ_MAT);
+	const float MaxODTNits = HDR_MAX_ODT_NITS;
+	const float scRGBScale = MaxODTNits / HDR_WHITE_POINT;
+	const float scRGBScaleInv = HDR_WHITE_POINT / MaxODTNits;
+
+	// normalize to [0-1]
+	Color_scRGB *= scRGBScaleInv;
+
+	// convert primaries sRGB -> Rec2020 
+	float3 Color_Rec2020 = mul(sRGB_2_Rec2020, Color_scRGB);
+
+	// clamp negatives
+	Color_Rec2020 = max(Color_Rec2020, 0); 
+
+	// Gamma2-encode
+	return LinearToGamma2(Color_Rec2020);
+}
+
+float3 LinearToScRGB(float3 LinearColor_Rec2020)
+{
+	const float3x3 Rec2020_2_sRGB = mul(XYZ_2_sRGB_MAT, Rec2020_2_XYZ_MAT);
+	const float MaxODTNits = HDR_MAX_ODT_NITS;
+	const float scRGBScale = MaxODTNits / HDR_WHITE_POINT;
+
+	// Rec2020 -> sRGB primaries
+	float3 LinearColor_sRGB = mul(Rec2020_2_sRGB, LinearColor_Rec2020);
+
+	// Linear_sRGB -> scRGB
+	return LinearColor_sRGB * scRGBScale;
+}
diff --git a/Engine/Plugins/Runtime/MobileFSR/Shaders/Private/PostProcessMobileFFX_FSR.usf b/Engine/Plugins/Runtime/MobileFSR/Shaders/Private/PostProcessMobileFFX_FSR.usf
new file mode 100644
--- /dev/null
+++ b/Engine/Plugins/Runtime/MobileFSR/Shaders/Private/PostProcessMobileFFX_FSR.usf
@@ -0,0 +1,161 @@
+// Copyright Epic Games, Inc. All Rights Reserved.
+//------------------------------------------------------------------------------
+//
+// Copyright (c) 2021 Advanced Micro Devices, Inc. All rights reserved.
+// Permission is hereby granted, free of charge, to any person obtaining a copy
+// of this software and associated documentation files(the "Software"), to deal
+// in the Software without restriction, including without limitation the rights
+// to use, copy, modify, merge, publish, distribute, sublicense, and / or sell
+// copies of the Software, and to permit persons to whom the Software is
+// furnished to do so, subject to the following conditions :
+// The above copyright notice and this permission notice shall be included in
+// all copies or substantial portions of the Software.
+// THE SOFTWARE IS PROVIDED "AS IS", WITHOUT WARRANTY OF ANY KIND, EXPRESS OR
+// IMPLIED, INCLUDING BUT NOT LIMITED TO THE WARRANTIES OF MERCHANTABILITY,
+// FITNESS FOR A PARTICULAR PURPOSE AND NONINFRINGEMENT.IN NO EVENT SHALL THE
+// AUTHORS OR COPYRIGHT HOLDERS BE LIABLE FOR ANY CLAIM, DAMAGES OR OTHER
+// LIABILITY, WHETHER IN AN ACTION OF CONTRACT, TORT OR OTHERWISE, ARISING FROM,
+// OUT OF OR IN CONNECTION WITH THE SOFTWARE OR THE USE OR OTHER DEALINGS IN
+// THE SOFTWARE.
+//------------------------------------------------------------------------------
+
+#include "/Engine/Private/Common.ush"
+#include "/Engine/Private/ScreenPass.ush"
+
+
+// =====================================================================================
+//
+// SHADER RESOURCES
+//
+// =====================================================================================
+uint4 Const0;
+uint4 Const1;
+uint4 Const2;
+uint4 Const3;
+
+SamplerState samLinearClamp;
+
+// UE4 GrainIntensity
+float4 GrainRandomFull; // zw:unused
+float4 GrainScaleBiasJitter; // x:Scale y:Bias z:Jitter
+
+float2 VPColor_ExtentInverse;
+float2 VPColor_ViewportMin;
+
+// =====================================================================================
+//
+// FIDELITYFX SETUP
+//
+// =====================================================================================
+#define A_GPU  1
+#define A_HLSL 1
+
+// Note: Input is expected to be in Gamma2-encoding.
+// LDR: This is already the case.
+// HDR: A ColorConversionCS runs before FSR to provide Gamma2 input to FSR passes
+#if ENABLE_FP16
+	#define A_HALF
+	#include "ffx_a.ush"
+	Texture2D<AH4> InputTexture;
+	RWTexture2D<AH4> OutputTexture;
+	#define FSR_EASU_H 1
+	AH4 FsrEasuRH(AF2 p) { AH4 res = InputTexture.GatherRed  (samLinearClamp, p, ASU2(0, 0)); return res; }
+	AH4 FsrEasuGH(AF2 p) { AH4 res = InputTexture.GatherGreen(samLinearClamp, p, ASU2(0, 0)); return res; }
+	AH4 FsrEasuBH(AF2 p) { AH4 res = InputTexture.GatherBlue (samLinearClamp, p, ASU2(0, 0)); return res; }
+	AH4 FsrEasuRHO(AF2 p, ASW2 o) { AH4 res = InputTexture.GatherRed  (samLinearClamp, p, o); return res; }
+	AH4 FsrEasuGHO(AF2 p, ASW2 o) { AH4 res = InputTexture.GatherGreen(samLinearClamp, p, o); return res; }
+	AH4 FsrEasuBHO(AF2 p, ASW2 o) { AH4 res = InputTexture.GatherBlue (samLinearClamp, p, o); return res; }
+	AH3 FsrEasuSampleH(AF2 p) { AH3 res = InputTexture.SampleLevel(samLinearClamp, p, 0).xyz; return res; }
+#else
+	#include "ffx_a.ush"
+	Texture2D InputTexture;
+	RWTexture2D<float4> OutputTexture;
+	#define FSR_EASU_F 1
+	AF4 FsrEasuRF(AF2 p) { AF4 res = InputTexture.GatherRed  (samLinearClamp, p, ASU2(0, 0)); return res; }
+	AF4 FsrEasuGF(AF2 p) { AF4 res = InputTexture.GatherGreen(samLinearClamp, p, ASU2(0, 0)); return res; }
+	AF4 FsrEasuBF(AF2 p) { AF4 res = InputTexture.GatherBlue (samLinearClamp, p, ASU2(0, 0)); return res; }
+#endif
+
+// should be included after ffx_a.ush, defines FSR_OUTPUTDEVICE
+#include "PostProcessMobileFFX_Common.ush"
+
+#include "ffx_fsr1.ush"
+
+// =====================================================================================
+// 
+// FIDELITYFX SUPER RESOLUTION
+// 
+// =====================================================================================
+
+#if ENABLE_FP16
+AH4 FidelityFXSuperResolution(AU2 gxy)
+#else
+AF4 FidelityFXSuperResolution(AU2 gxy)
+#endif
+{
+
+// 
+// FSR / EASU()
+//
+#if ENABLE_FP16
+	AH3 Gamma2Color = AH3(0,0,0);
+	#if ES3_1_PROFILE	
+		FsrEasuL(Gamma2Color, gxy, Const0, Const1, Const2, Const3);
+	#else
+		FsrEasuH(Gamma2Color, gxy, Const0, Const1, Const2, Const3);
+	#endif
+#else
+	AF3 Gamma2Color = AF3(0,0,0);
+	FsrEasuF(Gamma2Color, gxy, Const0, Const1, Const2, Const3);
+#endif // ENABLE_FP16
+
+	AF3 OutColor = AF3(Gamma2Color);
+
+#if ENABLE_FP16
+	return AH4(OutColor, 1);
+#else
+	return AF4(OutColor, 1);
+#endif
+}
+
+
+
+// =====================================================================================
+//
+// ENTRY POINTS
+//
+// =====================================================================================
+#if COMPUTE_SHADER
+[numthreads(THREADGROUP_SIZEX, THREADGROUP_SIZEY, THREADGROUP_SIZEZ)] 
+void MainCS(uint3 LocalThreadId : SV_GroupThreadID, uint3 WorkGroupId : SV_GroupID, uint3 Dtid : SV_DispatchThreadID)
+{
+	// Do remapping of local xy in workgroup for a more PS-like swizzle pattern.
+	// Threadgroup of 64 threads, each working on a cache-friendly 2x2 region (= 16x16 workgroup if 1px/thread assumed)
+	AU2 gxy = ARmp8x8(LocalThreadId.x) + AU2(WorkGroupId.x << 4u, WorkGroupId.y << 4u);
+
+	OutputTexture[gxy + VPColor_ViewportMin] = FidelityFXSuperResolution(gxy);
+	
+	gxy.x += 8u;
+	OutputTexture[gxy + VPColor_ViewportMin] = FidelityFXSuperResolution(gxy);
+	
+	gxy.y += 8u;
+	OutputTexture[gxy + VPColor_ViewportMin] = FidelityFXSuperResolution(gxy);
+	
+	gxy.x -= 8u;
+	OutputTexture[gxy + VPColor_ViewportMin] = FidelityFXSuperResolution(gxy);
+}  
+#endif // COMPUTE SHADER
+
+#if ENABLE_FP16
+void MainPS(noperspective float4 UVAndScreenPos : TEXCOORD0, float4 SvPosition : SV_POSITION, out AH4 OutColor : SV_Target0)
+{
+	AU2 gxy = AU2(SvPosition.xy - VPColor_ViewportMin);
+	OutColor = FidelityFXSuperResolution(gxy);
+}
+#else
+void MainPS(noperspective float4 UVAndScreenPos : TEXCOORD0, float4 SvPosition : SV_POSITION, out float4 OutColor : SV_Target0)
+{
+	AU2 gxy = AU2(SvPosition.xy - VPColor_ViewportMin);
+	OutColor = FidelityFXSuperResolution(gxy);
+}
+#endif
diff --git a/Engine/Plugins/Runtime/MobileFSR/Shaders/Private/PostProcessMobileFFX_RCAS.usf b/Engine/Plugins/Runtime/MobileFSR/Shaders/Private/PostProcessMobileFFX_RCAS.usf
new file mode 100644
--- /dev/null
+++ b/Engine/Plugins/Runtime/MobileFSR/Shaders/Private/PostProcessMobileFFX_RCAS.usf
@@ -0,0 +1,166 @@
+// Copyright Epic Games, Inc. All Rights Reserved.
+//------------------------------------------------------------------------------
+// Copyright (c) 2021 Advanced Micro Devices, Inc. All rights reserved.
+// Permission is hereby granted, free of charge, to any person obtaining a copy
+// of this software and associated documentation files(the "Software"), to deal
+// in the Software without restriction, including without limitation the rights
+// to use, copy, modify, merge, publish, distribute, sublicense, and / or sell
+// copies of the Software, and to permit persons to whom the Software is
+// furnished to do so, subject to the following conditions :
+// The above copyright notice and this permission notice shall be included in
+// all copies or substantial portions of the Software.
+// THE SOFTWARE IS PROVIDED "AS IS", WITHOUT WARRANTY OF ANY KIND, EXPRESS OR
+// IMPLIED, INCLUDING BUT NOT LIMITED TO THE WARRANTIES OF MERCHANTABILITY,
+// FITNESS FOR A PARTICULAR PURPOSE AND NONINFRINGEMENT.IN NO EVENT SHALL THE
+// AUTHORS OR COPYRIGHT HOLDERS BE LIABLE FOR ANY CLAIM, DAMAGES OR OTHER
+// LIABILITY, WHETHER IN AN ACTION OF CONTRACT, TORT OR OTHERWISE, ARISING FROM,
+// OUT OF OR IN CONNECTION WITH THE SOFTWARE OR THE USE OR OTHER DEALINGS IN
+// THE SOFTWARE.
+//------------------------------------------------------------------------------
+
+// FidelityFX Super Resolution - Robust Contrast Adaptive Sharpening
+
+#include "/Engine/Private/Common.ush"
+#include "/Engine/Private/ScreenPass.ush"
+
+
+// =====================================================================================
+//
+// SHADER RESOURCES
+//
+// =====================================================================================
+
+// RCAS params
+uint4 Const0;
+float2 VPColor_ExtentInverse;
+
+
+// =====================================================================================
+//
+// FIDELITYFX SETUP
+//
+// =====================================================================================
+#define A_HLSL 1
+#define A_GPU  1
+
+#if USE_RCAS_DENOISE 
+#define FSR_RCAS_DENOISE 1
+#endif
+
+#if ENABLE_FP16
+	#define FSR_RCAS_H   1
+	#define FSR_RCAS_HX2 1
+	#define A_HALF       1
+	#include "ffx_a.ush"
+
+	Texture2D<AH4>   InputTexture;
+	RWTexture2D<AH4> OutputTexture;
+	AH4 FsrRcasLoadH(min16int2 p)  {return InputTexture.Load(int3(p, 0));}
+	AH4 FsrRcasLoadHx2(min16int2 p){return InputTexture.Load(int3(p, 0));}
+#else
+	#define FSR_RCAS_F 1
+	#include "ffx_a.ush"
+
+	Texture2D           InputTexture;
+	RWTexture2D<float4> OutputTexture;
+	AF4 FsrRcasLoadF(int2 p){return InputTexture.Load(int3(p, 0));}
+#endif
+	
+// should be included after ffx_a.ush, defines FSR_OUTPUTDEVICE
+#include "PostProcessMobileFFX_Common.ush"
+
+
+// Note: Input is expected to be in Gamma2-encoding.
+// LDR: This is already the case.
+// HDR: A ColorConversionCS runs before FSR to provide Gamma2 input to FSR passes
+#if ENABLE_FP16
+	void FsrRcasInputH  (inout AH1 r,inout AH1 g,inout AH1 b){} // RCAS is done in perceptual space (sRGB~Gamma2)
+	void FsrRcasInputHx2(inout AH2 r,inout AH2 g,inout AH2 b){} // Hence, no need to modify r/g/b
+#else
+	void FsrRcasInputF(inout AF1 r, inout AF1 g, inout AF1 b){}
+#endif
+
+
+#include "ffx_fsr1.ush"
+
+
+// =====================================================================================
+// 
+// RCAS IMPLEMENTATION
+// 
+// =====================================================================================
+AF3 RCASPassCommon(AU2 gxy)
+{
+	//
+	// RCAS()
+	//
+	#if ENABLE_FP16 // half precision floating point
+
+		AH3 Gamma2Color = AH3(0, 0, 0);
+		FsrRcasH(Gamma2Color.r, Gamma2Color.g, Gamma2Color.b, gxy, Const0);
+	
+	#else // FP32 - full precision floating point
+
+		AF3 Gamma2Color = AF3(0, 0, 0);
+		FsrRcasF(Gamma2Color.r, Gamma2Color.g, Gamma2Color.b, gxy, Const0);
+
+	#endif // ENABLE_FP16
+
+	AF3 OutColor = AF3(Gamma2Color);
+
+	return OutColor;
+}
+
+
+
+// =====================================================================================
+//
+// ENTRY POINTS
+//
+// =====================================================================================
+#if COMPUTE_SHADER
+[numthreads(THREADGROUP_SIZEX, THREADGROUP_SIZEY, THREADGROUP_SIZEZ)]
+void MainCS(uint3 LocalThreadId : SV_GroupThreadID, uint3 WorkGroupId : SV_GroupID, uint3 Dtid : SV_DispatchThreadID)
+{
+	// Do remapping of local xy in workgroup for a more PS-like swizzle pattern.
+	AU2 gxy = ARmp8x8(LocalThreadId.x) + AU2(WorkGroupId.x << 4u, WorkGroupId.y << 4u);	
+	
+#if ENABLE_FP16
+	OutputTexture[gxy] = AH4(RCASPassCommon(gxy), 1);
+
+	gxy.x += 8u;
+	OutputTexture[gxy] = AH4(RCASPassCommon(gxy), 1);
+
+	gxy.y += 8u;
+	OutputTexture[gxy] = AH4(RCASPassCommon(gxy), 1);
+
+	gxy.x -= 8u;
+	OutputTexture[gxy] = AH4(RCASPassCommon(gxy), 1);
+#else
+	OutputTexture[gxy] = AF4(RCASPassCommon(gxy), 1);
+
+	gxy.x += 8u;
+	OutputTexture[gxy] = AF4(RCASPassCommon(gxy), 1);
+
+	gxy.y += 8u;
+	OutputTexture[gxy] = AF4(RCASPassCommon(gxy), 1);
+
+	gxy.x -= 8u;
+	OutputTexture[gxy] = AF4(RCASPassCommon(gxy), 1);
+#endif
+}
+#endif // COMPUTE_SHADER
+
+#if ENABLE_FP16
+void MainPS(noperspective float4 UVAndScreenPos : TEXCOORD0, float4 SvPosition : SV_POSITION, out AH4 OutColor : SV_Target0)
+{
+	AU2 gxy = AU2(SvPosition.xy);
+	OutColor = AH4(RCASPassCommon(gxy), 1);
+}
+#else
+void MainPS(noperspective float4 UVAndScreenPos : TEXCOORD0, float4 SvPosition : SV_POSITION, out float4 OutColor : SV_Target0)
+{
+	AU2 gxy = AU2(SvPosition.xy);
+	OutColor = AF4(RCASPassCommon(gxy), 1);
+}
+#endif
\ No newline at end of file
