Removed Mobile FSR from Source/Shaders

diff --git a/Engine/Plugins/Runtime/MobileFSR/Shaders/Private/ffx_a.ush b/Engine/Plugins/Runtime/MobileFSR/Shaders/Private/ffx_a.ush
new file mode 100644
--- /dev/null
+++ b/Engine/Plugins/Runtime/MobileFSR/Shaders/Private/ffx_a.ush
@@ -0,0 +1,2656 @@
+//==============================================================================================================================
+//
+//                                               [A] SHADER PORTABILITY 1.20210629
+//
+//==============================================================================================================================
+// FidelityFX Super Resolution Sample
+//
+// Copyright (c) 2021 Advanced Micro Devices, Inc. All rights reserved.
+// Permission is hereby granted, free of charge, to any person obtaining a copy
+// of this software and associated documentation files(the "Software"), to deal
+// in the Software without restriction, including without limitation the rights
+// to use, copy, modify, merge, publish, distribute, sublicense, and / or sell
+// copies of the Software, and to permit persons to whom the Software is
+// furnished to do so, subject to the following conditions :
+// The above copyright notice and this permission notice shall be included in
+// all copies or substantial portions of the Software.
+// THE SOFTWARE IS PROVIDED "AS IS", WITHOUT WARRANTY OF ANY KIND, EXPRESS OR
+// IMPLIED, INCLUDING BUT NOT LIMITED TO THE WARRANTIES OF MERCHANTABILITY,
+// FITNESS FOR A PARTICULAR PURPOSE AND NONINFRINGEMENT.IN NO EVENT SHALL THE
+// AUTHORS OR COPYRIGHT HOLDERS BE LIABLE FOR ANY CLAIM, DAMAGES OR OTHER
+// LIABILITY, WHETHER IN AN ACTION OF CONTRACT, TORT OR OTHERWISE, ARISING FROM,
+// OUT OF OR IN CONNECTION WITH THE SOFTWARE OR THE USE OR OTHER DEALINGS IN
+// THE SOFTWARE.
+//------------------------------------------------------------------------------------------------------------------------------
+// MIT LICENSE
+// ===========
+// Copyright (c) 2014 Michal Drobot (for concepts used in "FLOAT APPROXIMATIONS").
+// -----------
+// Permission is hereby granted, free of charge, to any person obtaining a copy of this software and associated documentation
+// files (the "Software"), to deal in the Software without restriction, including without limitation the rights to use, copy,
+// modify, merge, publish, distribute, sublicense, and/or sell copies of the Software, and to permit persons to whom the
+// Software is furnished to do so, subject to the following conditions:
+// -----------
+// The above copyright notice and this permission notice shall be included in all copies or substantial portions of the
+// Software.
+// -----------
+// THE SOFTWARE IS PROVIDED "AS IS", WITHOUT WARRANTY OF ANY KIND, EXPRESS OR IMPLIED, INCLUDING BUT NOT LIMITED TO THE
+// WARRANTIES OF MERCHANTABILITY, FITNESS FOR A PARTICULAR PURPOSE AND NONINFRINGEMENT.  IN NO EVENT SHALL THE AUTHORS OR
+// COPYRIGHT HOLDERS BE LIABLE FOR ANY CLAIM, DAMAGES OR OTHER LIABILITY, WHETHER IN AN ACTION OF CONTRACT, TORT OR OTHERWISE,
+// ARISING FROM, OUT OF OR IN CONNECTION WITH THE SOFTWARE OR THE USE OR OTHER DEALINGS IN THE SOFTWARE.
+//------------------------------------------------------------------------------------------------------------------------------
+// ABOUT
+// =====
+// Common central point for high-level shading language and C portability for various shader headers.
+//------------------------------------------------------------------------------------------------------------------------------
+// DEFINES
+// =======
+// A_CPU ..... Include the CPU related code.
+// A_GPU ..... Include the GPU related code.
+// A_GLSL .... Using GLSL.
+// A_HLSL .... Using HLSL.
+// A_HLSL_6_2  Using HLSL 6.2 with new 'uint16_t' and related types (requires '-enable-16bit-types').
+// A_NO_16_BIT_CAST Don't use instructions that are not availabe in SPIR-V (needed for running A_HLSL_6_2 on Vulkan)
+// A_GCC ..... Using a GCC compatible compiler (else assume MSVC compatible compiler by default).
+// =======
+// A_BYTE .... Support 8-bit integer.
+// A_HALF .... Support 16-bit integer and floating point.
+// A_LONG .... Support 64-bit integer.
+// A_DUBL .... Support 64-bit floating point.
+// =======
+// A_WAVE .... Support wave-wide operations.
+//------------------------------------------------------------------------------------------------------------------------------
+// To get #include "ffx_a.ush" working in GLSL use '#extension GL_GOOGLE_include_directive:require'.
+//------------------------------------------------------------------------------------------------------------------------------
+// SIMPLIFIED TYPE SYSTEM
+// ======================
+//  - All ints will be unsigned with exception of when signed is required.
+//  - Type naming simplified and shortened "A<type><#components>",
+//     - H = 16-bit float (half)
+//     - F = 32-bit float (float)
+//     - D = 64-bit float (double)
+//     - P = 1-bit integer (predicate, not using bool because 'B' is used for byte)
+//     - B = 8-bit integer (byte)
+//     - W = 16-bit integer (word)
+//     - U = 32-bit integer (unsigned)
+//     - L = 64-bit integer (long)
+//  - Using "AS<type><#components>" for signed when required.
+//------------------------------------------------------------------------------------------------------------------------------
+// TODO
+// ====
+//  - Make sure 'ALerp*(a,b,m)' does 'b*m+(-a*m+a)' (2 ops).
+//------------------------------------------------------------------------------------------------------------------------------
+// CHANGE LOG
+// ==========
+// 20200914 - Expanded wave ops and prx code.
+// 20200713 - Added [ZOL] section, fixed serious bugs in sRGB and Rec.709 color conversion code, etc.
+//==============================================================================================================================
+////////////////////////////////////////////////////////////////////////////////////////////////////////////////////////////////
+////////////////////////////////////////////////////////////////////////////////////////////////////////////////////////////////
+//_____________________________________________________________/\_______________________________________________________________
+//==============================================================================================================================
+//                                                           COMMON
+//==============================================================================================================================
+#define A_2PI 6.28318530718
+////////////////////////////////////////////////////////////////////////////////////////////////////////////////////////////////
+////////////////////////////////////////////////////////////////////////////////////////////////////////////////////////////////
+////////////////////////////////////////////////////////////////////////////////////////////////////////////////////////////////
+////////////////////////////////////////////////////////////////////////////////////////////////////////////////////////////////
+////////////////////////////////////////////////////////////////////////////////////////////////////////////////////////////////
+////////////////////////////////////////////////////////////////////////////////////////////////////////////////////////////////
+//_____________________________________________________________/\_______________________________________________________________
+//==============================================================================================================================
+//
+//
+//                                                             CPU
+//
+//
+//==============================================================================================================================
+#ifdef A_CPU
+ // Supporting user defined overrides.
+ #ifndef A_RESTRICT
+  #define A_RESTRICT __restrict
+ #endif
+//------------------------------------------------------------------------------------------------------------------------------
+ #ifndef A_STATIC
+  #define A_STATIC static
+ #endif
+//------------------------------------------------------------------------------------------------------------------------------
+ // Same types across CPU and GPU.
+ // Predicate uses 32-bit integer (C friendly bool).
+ typedef uint32_t AP1;
+ typedef float AF1;
+ typedef double AD1;
+ typedef uint8_t AB1;
+ typedef uint16_t AW1;
+ typedef uint32_t AU1;
+ typedef uint64_t AL1;
+ typedef int8_t ASB1;
+ typedef int16_t ASW1;
+ typedef int32_t ASU1;
+ typedef int64_t ASL1;
+//------------------------------------------------------------------------------------------------------------------------------
+ #define AD1_(a) ((AD1)(a))
+ #define AF1_(a) ((AF1)(a))
+ #define AL1_(a) ((AL1)(a))
+ #define AU1_(a) ((AU1)(a))
+//------------------------------------------------------------------------------------------------------------------------------
+ #define ASL1_(a) ((ASL1)(a))
+ #define ASU1_(a) ((ASU1)(a))
+//------------------------------------------------------------------------------------------------------------------------------
+ A_STATIC AU1 AU1_AF1(AF1 a){union{AF1 f;AU1 u;}bits;bits.f=a;return bits.u;}
+//------------------------------------------------------------------------------------------------------------------------------
+ #define A_TRUE 1
+ #define A_FALSE 0
+////////////////////////////////////////////////////////////////////////////////////////////////////////////////////////////////
+////////////////////////////////////////////////////////////////////////////////////////////////////////////////////////////////
+////////////////////////////////////////////////////////////////////////////////////////////////////////////////////////////////
+////////////////////////////////////////////////////////////////////////////////////////////////////////////////////////////////
+//_____________________________________________________________/\_______________________________________________________________
+//==============================================================================================================================
+//
+//                                                       CPU/GPU PORTING
+//
+//------------------------------------------------------------------------------------------------------------------------------
+// Get CPU and GPU to share all setup code, without duplicate code paths.
+// This uses a lower-case prefix for special vector constructs.
+//  - In C restrict pointers are used.
+//  - In the shading language, in/inout/out arguments are used.
+// This depends on the ability to access a vector value in both languages via array syntax (aka color[2]).
+//==============================================================================================================================
+////////////////////////////////////////////////////////////////////////////////////////////////////////////////////////////////
+////////////////////////////////////////////////////////////////////////////////////////////////////////////////////////////////
+//_____________________________________________________________/\_______________________________________________________________
+//==============================================================================================================================
+//                                     VECTOR ARGUMENT/RETURN/INITIALIZATION PORTABILITY
+//==============================================================================================================================
+ #define retAD2 AD1 *A_RESTRICT
+ #define retAD3 AD1 *A_RESTRICT
+ #define retAD4 AD1 *A_RESTRICT
+ #define retAF2 AF1 *A_RESTRICT
+ #define retAF3 AF1 *A_RESTRICT
+ #define retAF4 AF1 *A_RESTRICT
+ #define retAL2 AL1 *A_RESTRICT
+ #define retAL3 AL1 *A_RESTRICT
+ #define retAL4 AL1 *A_RESTRICT
+ #define retAU2 AU1 *A_RESTRICT
+ #define retAU3 AU1 *A_RESTRICT
+ #define retAU4 AU1 *A_RESTRICT
+//------------------------------------------------------------------------------------------------------------------------------
+ #define inAD2 AD1 *A_RESTRICT
+ #define inAD3 AD1 *A_RESTRICT
+ #define inAD4 AD1 *A_RESTRICT
+ #define inAF2 AF1 *A_RESTRICT
+ #define inAF3 AF1 *A_RESTRICT
+ #define inAF4 AF1 *A_RESTRICT
+ #define inAL2 AL1 *A_RESTRICT
+ #define inAL3 AL1 *A_RESTRICT
+ #define inAL4 AL1 *A_RESTRICT
+ #define inAU2 AU1 *A_RESTRICT
+ #define inAU3 AU1 *A_RESTRICT
+ #define inAU4 AU1 *A_RESTRICT
+//------------------------------------------------------------------------------------------------------------------------------
+ #define inoutAD2 AD1 *A_RESTRICT
+ #define inoutAD3 AD1 *A_RESTRICT
+ #define inoutAD4 AD1 *A_RESTRICT
+ #define inoutAF2 AF1 *A_RESTRICT
+ #define inoutAF3 AF1 *A_RESTRICT
+ #define inoutAF4 AF1 *A_RESTRICT
+ #define inoutAL2 AL1 *A_RESTRICT
+ #define inoutAL3 AL1 *A_RESTRICT
+ #define inoutAL4 AL1 *A_RESTRICT
+ #define inoutAU2 AU1 *A_RESTRICT
+ #define inoutAU3 AU1 *A_RESTRICT
+ #define inoutAU4 AU1 *A_RESTRICT
+//------------------------------------------------------------------------------------------------------------------------------
+ #define outAD2 AD1 *A_RESTRICT
+ #define outAD3 AD1 *A_RESTRICT
+ #define outAD4 AD1 *A_RESTRICT
+ #define outAF2 AF1 *A_RESTRICT
+ #define outAF3 AF1 *A_RESTRICT
+ #define outAF4 AF1 *A_RESTRICT
+ #define outAL2 AL1 *A_RESTRICT
+ #define outAL3 AL1 *A_RESTRICT
+ #define outAL4 AL1 *A_RESTRICT
+ #define outAU2 AU1 *A_RESTRICT
+ #define outAU3 AU1 *A_RESTRICT
+ #define outAU4 AU1 *A_RESTRICT
+//------------------------------------------------------------------------------------------------------------------------------
+ #define varAD2(x) AD1 x[2]
+ #define varAD3(x) AD1 x[3]
+ #define varAD4(x) AD1 x[4]
+ #define varAF2(x) AF1 x[2]
+ #define varAF3(x) AF1 x[3]
+ #define varAF4(x) AF1 x[4]
+ #define varAL2(x) AL1 x[2]
+ #define varAL3(x) AL1 x[3]
+ #define varAL4(x) AL1 x[4]
+ #define varAU2(x) AU1 x[2]
+ #define varAU3(x) AU1 x[3]
+ #define varAU4(x) AU1 x[4]
+//------------------------------------------------------------------------------------------------------------------------------
+ #define initAD2(x,y) {x,y}
+ #define initAD3(x,y,z) {x,y,z}
+ #define initAD4(x,y,z,w) {x,y,z,w}
+ #define initAF2(x,y) {x,y}
+ #define initAF3(x,y,z) {x,y,z}
+ #define initAF4(x,y,z,w) {x,y,z,w}
+ #define initAL2(x,y) {x,y}
+ #define initAL3(x,y,z) {x,y,z}
+ #define initAL4(x,y,z,w) {x,y,z,w}
+ #define initAU2(x,y) {x,y}
+ #define initAU3(x,y,z) {x,y,z}
+ #define initAU4(x,y,z,w) {x,y,z,w}
+////////////////////////////////////////////////////////////////////////////////////////////////////////////////////////////////
+////////////////////////////////////////////////////////////////////////////////////////////////////////////////////////////////
+//_____________________________________________________________/\_______________________________________________________________
+//==============================================================================================================================
+//                                                     SCALAR RETURN OPS
+//------------------------------------------------------------------------------------------------------------------------------
+// TODO
+// ====
+//  - Replace transcendentals with manual versions. 
+//==============================================================================================================================
+ #ifdef A_GCC
+  A_STATIC AD1 AAbsD1(AD1 a){return __builtin_fabs(a);}
+  A_STATIC AF1 AAbsF1(AF1 a){return __builtin_fabsf(a);}
+  A_STATIC AU1 AAbsSU1(AU1 a){return AU1_(__builtin_abs(ASU1_(a)));}
+  A_STATIC AL1 AAbsSL1(AL1 a){return AL1_(__builtin_llabs(ASL1_(a)));}
+ #else
+  A_STATIC AD1 AAbsD1(AD1 a){return fabs(a);}
+  A_STATIC AF1 AAbsF1(AF1 a){return fabsf(a);}
+  A_STATIC AU1 AAbsSU1(AU1 a){return AU1_(abs(ASU1_(a)));}
+  A_STATIC AL1 AAbsSL1(AL1 a){return AL1_(labs((long)ASL1_(a)));}
+ #endif
+//------------------------------------------------------------------------------------------------------------------------------
+ #ifdef A_GCC
+  A_STATIC AD1 ACosD1(AD1 a){return __builtin_cos(a);}
+  A_STATIC AF1 ACosF1(AF1 a){return __builtin_cosf(a);}
+ #else
+  A_STATIC AD1 ACosD1(AD1 a){return cos(a);}
+  A_STATIC AF1 ACosF1(AF1 a){return cosf(a);}
+ #endif
+//------------------------------------------------------------------------------------------------------------------------------
+ A_STATIC AD1 ADotD2(inAD2 a,inAD2 b){return a[0]*b[0]+a[1]*b[1];}
+ A_STATIC AD1 ADotD3(inAD3 a,inAD3 b){return a[0]*b[0]+a[1]*b[1]+a[2]*b[2];}
+ A_STATIC AD1 ADotD4(inAD4 a,inAD4 b){return a[0]*b[0]+a[1]*b[1]+a[2]*b[2]+a[3]*b[3];}
+ A_STATIC AF1 ADotF2(inAF2 a,inAF2 b){return a[0]*b[0]+a[1]*b[1];}
+ A_STATIC AF1 ADotF3(inAF3 a,inAF3 b){return a[0]*b[0]+a[1]*b[1]+a[2]*b[2];}
+ A_STATIC AF1 ADotF4(inAF4 a,inAF4 b){return a[0]*b[0]+a[1]*b[1]+a[2]*b[2]+a[3]*b[3];}
+//------------------------------------------------------------------------------------------------------------------------------
+ #ifdef A_GCC
+  A_STATIC AD1 AExp2D1(AD1 a){return __builtin_exp2(a);}
+  A_STATIC AF1 AExp2F1(AF1 a){return __builtin_exp2f(a);}
+ #else
+  A_STATIC AD1 AExp2D1(AD1 a){return exp2(a);}
+  A_STATIC AF1 AExp2F1(AF1 a){return exp2f(a);}
+ #endif
+//------------------------------------------------------------------------------------------------------------------------------
+ #ifdef A_GCC
+  A_STATIC AD1 AFloorD1(AD1 a){return __builtin_floor(a);}
+  A_STATIC AF1 AFloorF1(AF1 a){return __builtin_floorf(a);}
+ #else
+  A_STATIC AD1 AFloorD1(AD1 a){return floor(a);}
+  A_STATIC AF1 AFloorF1(AF1 a){return floorf(a);}
+ #endif
+//------------------------------------------------------------------------------------------------------------------------------
+ A_STATIC AD1 ALerpD1(AD1 a,AD1 b,AD1 c){return b*c+(-a*c+a);}
+ A_STATIC AF1 ALerpF1(AF1 a,AF1 b,AF1 c){return b*c+(-a*c+a);}
+//------------------------------------------------------------------------------------------------------------------------------
+ #ifdef A_GCC
+  A_STATIC AD1 ALog2D1(AD1 a){return __builtin_log2(a);}
+  A_STATIC AF1 ALog2F1(AF1 a){return __builtin_log2f(a);}
+ #else
+  A_STATIC AD1 ALog2D1(AD1 a){return log2(a);}
+  A_STATIC AF1 ALog2F1(AF1 a){return log2f(a);}
+ #endif
+//------------------------------------------------------------------------------------------------------------------------------
+ A_STATIC AD1 AMaxD1(AD1 a,AD1 b){return a>b?a:b;}
+ A_STATIC AF1 AMaxF1(AF1 a,AF1 b){return a>b?a:b;}
+ A_STATIC AL1 AMaxL1(AL1 a,AL1 b){return a>b?a:b;}
+ A_STATIC AU1 AMaxU1(AU1 a,AU1 b){return a>b?a:b;}
+//------------------------------------------------------------------------------------------------------------------------------
+ // These follow the convention that A integer types don't have signage, until they are operated on. 
+ A_STATIC AL1 AMaxSL1(AL1 a,AL1 b){return (ASL1_(a)>ASL1_(b))?a:b;}
+ A_STATIC AU1 AMaxSU1(AU1 a,AU1 b){return (ASU1_(a)>ASU1_(b))?a:b;}
+//------------------------------------------------------------------------------------------------------------------------------
+ A_STATIC AD1 AMinD1(AD1 a,AD1 b){return a<b?a:b;}
+ A_STATIC AF1 AMinF1(AF1 a,AF1 b){return a<b?a:b;}
+ A_STATIC AL1 AMinL1(AL1 a,AL1 b){return a<b?a:b;}
+ A_STATIC AU1 AMinU1(AU1 a,AU1 b){return a<b?a:b;}
+//------------------------------------------------------------------------------------------------------------------------------
+ A_STATIC AL1 AMinSL1(AL1 a,AL1 b){return (ASL1_(a)<ASL1_(b))?a:b;}
+ A_STATIC AU1 AMinSU1(AU1 a,AU1 b){return (ASU1_(a)<ASU1_(b))?a:b;}
+//------------------------------------------------------------------------------------------------------------------------------
+ A_STATIC AD1 ARcpD1(AD1 a){return 1.0/a;}
+ A_STATIC AF1 ARcpF1(AF1 a){return 1.0f/a;}
+//------------------------------------------------------------------------------------------------------------------------------
+ A_STATIC AL1 AShrSL1(AL1 a,AL1 b){return AL1_(ASL1_(a)>>ASL1_(b));}
+ A_STATIC AU1 AShrSU1(AU1 a,AU1 b){return AU1_(ASU1_(a)>>ASU1_(b));}
+//------------------------------------------------------------------------------------------------------------------------------
+ #ifdef A_GCC
+  A_STATIC AD1 ASinD1(AD1 a){return __builtin_sin(a);}
+  A_STATIC AF1 ASinF1(AF1 a){return __builtin_sinf(a);}
+ #else
+  A_STATIC AD1 ASinD1(AD1 a){return sin(a);}
+  A_STATIC AF1 ASinF1(AF1 a){return sinf(a);}
+ #endif
+//------------------------------------------------------------------------------------------------------------------------------
+ #ifdef A_GCC
+  A_STATIC AD1 ASqrtD1(AD1 a){return __builtin_sqrt(a);}
+  A_STATIC AF1 ASqrtF1(AF1 a){return __builtin_sqrtf(a);}
+ #else
+  A_STATIC AD1 ASqrtD1(AD1 a){return sqrt(a);}
+  A_STATIC AF1 ASqrtF1(AF1 a){return sqrtf(a);}
+ #endif
+////////////////////////////////////////////////////////////////////////////////////////////////////////////////////////////////
+////////////////////////////////////////////////////////////////////////////////////////////////////////////////////////////////
+//_____________________________________________________________/\_______________________________________________________________
+//==============================================================================================================================
+//                                               SCALAR RETURN OPS - DEPENDENT
+//==============================================================================================================================
+ A_STATIC AD1 AClampD1(AD1 x,AD1 n,AD1 m){return AMaxD1(n,AMinD1(x,m));}
+ A_STATIC AF1 AClampF1(AF1 x,AF1 n,AF1 m){return AMaxF1(n,AMinF1(x,m));}
+//------------------------------------------------------------------------------------------------------------------------------
+ A_STATIC AD1 AFractD1(AD1 a){return a-AFloorD1(a);}
+ A_STATIC AF1 AFractF1(AF1 a){return a-AFloorF1(a);}
+//------------------------------------------------------------------------------------------------------------------------------
+ A_STATIC AD1 APowD1(AD1 a,AD1 b){return AExp2D1(b*ALog2D1(a));}
+ A_STATIC AF1 APowF1(AF1 a,AF1 b){return AExp2F1(b*ALog2F1(a));}
+//------------------------------------------------------------------------------------------------------------------------------
+ A_STATIC AD1 ARsqD1(AD1 a){return ARcpD1(ASqrtD1(a));}
+ A_STATIC AF1 ARsqF1(AF1 a){return ARcpF1(ASqrtF1(a));}
+//------------------------------------------------------------------------------------------------------------------------------
+ A_STATIC AD1 ASatD1(AD1 a){return AMinD1(1.0,AMaxD1(0.0,a));}
+ A_STATIC AF1 ASatF1(AF1 a){return AMinF1(1.0f,AMaxF1(0.0f,a));}
+////////////////////////////////////////////////////////////////////////////////////////////////////////////////////////////////
+////////////////////////////////////////////////////////////////////////////////////////////////////////////////////////////////
+//_____________________________________________________________/\_______________________________________________________________
+//==============================================================================================================================
+//                                                         VECTOR OPS
+//------------------------------------------------------------------------------------------------------------------------------
+// These are added as needed for production or prototyping, so not necessarily a complete set.
+// They follow a convention of taking in a destination and also returning the destination value to increase utility.
+//==============================================================================================================================
+ A_STATIC retAD2 opAAbsD2(outAD2 d,inAD2 a){d[0]=AAbsD1(a[0]);d[1]=AAbsD1(a[1]);return d;}
+ A_STATIC retAD3 opAAbsD3(outAD3 d,inAD3 a){d[0]=AAbsD1(a[0]);d[1]=AAbsD1(a[1]);d[2]=AAbsD1(a[2]);return d;}
+ A_STATIC retAD4 opAAbsD4(outAD4 d,inAD4 a){d[0]=AAbsD1(a[0]);d[1]=AAbsD1(a[1]);d[2]=AAbsD1(a[2]);d[3]=AAbsD1(a[3]);return d;}
+//------------------------------------------------------------------------------------------------------------------------------
+ A_STATIC retAF2 opAAbsF2(outAF2 d,inAF2 a){d[0]=AAbsF1(a[0]);d[1]=AAbsF1(a[1]);return d;}
+ A_STATIC retAF3 opAAbsF3(outAF3 d,inAF3 a){d[0]=AAbsF1(a[0]);d[1]=AAbsF1(a[1]);d[2]=AAbsF1(a[2]);return d;}
+ A_STATIC retAF4 opAAbsF4(outAF4 d,inAF4 a){d[0]=AAbsF1(a[0]);d[1]=AAbsF1(a[1]);d[2]=AAbsF1(a[2]);d[3]=AAbsF1(a[3]);return d;}
+//==============================================================================================================================
+ A_STATIC retAD2 opAAddD2(outAD2 d,inAD2 a,inAD2 b){d[0]=a[0]+b[0];d[1]=a[1]+b[1];return d;}
+ A_STATIC retAD3 opAAddD3(outAD3 d,inAD3 a,inAD3 b){d[0]=a[0]+b[0];d[1]=a[1]+b[1];d[2]=a[2]+b[2];return d;}
+ A_STATIC retAD4 opAAddD4(outAD4 d,inAD4 a,inAD4 b){d[0]=a[0]+b[0];d[1]=a[1]+b[1];d[2]=a[2]+b[2];d[3]=a[3]+b[3];return d;}
+//------------------------------------------------------------------------------------------------------------------------------
+ A_STATIC retAF2 opAAddF2(outAF2 d,inAF2 a,inAF2 b){d[0]=a[0]+b[0];d[1]=a[1]+b[1];return d;}
+ A_STATIC retAF3 opAAddF3(outAF3 d,inAF3 a,inAF3 b){d[0]=a[0]+b[0];d[1]=a[1]+b[1];d[2]=a[2]+b[2];return d;}
+ A_STATIC retAF4 opAAddF4(outAF4 d,inAF4 a,inAF4 b){d[0]=a[0]+b[0];d[1]=a[1]+b[1];d[2]=a[2]+b[2];d[3]=a[3]+b[3];return d;}
+//==============================================================================================================================
+ A_STATIC retAD2 opAAddOneD2(outAD2 d,inAD2 a,AD1 b){d[0]=a[0]+b;d[1]=a[1]+b;return d;}
+ A_STATIC retAD3 opAAddOneD3(outAD3 d,inAD3 a,AD1 b){d[0]=a[0]+b;d[1]=a[1]+b;d[2]=a[2]+b;return d;}
+ A_STATIC retAD4 opAAddOneD4(outAD4 d,inAD4 a,AD1 b){d[0]=a[0]+b;d[1]=a[1]+b;d[2]=a[2]+b;d[3]=a[3]+b;return d;}
+//------------------------------------------------------------------------------------------------------------------------------
+ A_STATIC retAF2 opAAddOneF2(outAF2 d,inAF2 a,AF1 b){d[0]=a[0]+b;d[1]=a[1]+b;return d;}
+ A_STATIC retAF3 opAAddOneF3(outAF3 d,inAF3 a,AF1 b){d[0]=a[0]+b;d[1]=a[1]+b;d[2]=a[2]+b;return d;}
+ A_STATIC retAF4 opAAddOneF4(outAF4 d,inAF4 a,AF1 b){d[0]=a[0]+b;d[1]=a[1]+b;d[2]=a[2]+b;d[3]=a[3]+b;return d;}
+//==============================================================================================================================
+ A_STATIC retAD2 opACpyD2(outAD2 d,inAD2 a){d[0]=a[0];d[1]=a[1];return d;}
+ A_STATIC retAD3 opACpyD3(outAD3 d,inAD3 a){d[0]=a[0];d[1]=a[1];d[2]=a[2];return d;}
+ A_STATIC retAD4 opACpyD4(outAD4 d,inAD4 a){d[0]=a[0];d[1]=a[1];d[2]=a[2];d[3]=a[3];return d;}
+//------------------------------------------------------------------------------------------------------------------------------
+ A_STATIC retAF2 opACpyF2(outAF2 d,inAF2 a){d[0]=a[0];d[1]=a[1];return d;}
+ A_STATIC retAF3 opACpyF3(outAF3 d,inAF3 a){d[0]=a[0];d[1]=a[1];d[2]=a[2];return d;}
+ A_STATIC retAF4 opACpyF4(outAF4 d,inAF4 a){d[0]=a[0];d[1]=a[1];d[2]=a[2];d[3]=a[3];return d;}
+//==============================================================================================================================
+ A_STATIC retAD2 opALerpD2(outAD2 d,inAD2 a,inAD2 b,inAD2 c){d[0]=ALerpD1(a[0],b[0],c[0]);d[1]=ALerpD1(a[1],b[1],c[1]);return d;}
+ A_STATIC retAD3 opALerpD3(outAD3 d,inAD3 a,inAD3 b,inAD3 c){d[0]=ALerpD1(a[0],b[0],c[0]);d[1]=ALerpD1(a[1],b[1],c[1]);d[2]=ALerpD1(a[2],b[2],c[2]);return d;}
+ A_STATIC retAD4 opALerpD4(outAD4 d,inAD4 a,inAD4 b,inAD4 c){d[0]=ALerpD1(a[0],b[0],c[0]);d[1]=ALerpD1(a[1],b[1],c[1]);d[2]=ALerpD1(a[2],b[2],c[2]);d[3]=ALerpD1(a[3],b[3],c[3]);return d;}
+//------------------------------------------------------------------------------------------------------------------------------
+ A_STATIC retAF2 opALerpF2(outAF2 d,inAF2 a,inAF2 b,inAF2 c){d[0]=ALerpF1(a[0],b[0],c[0]);d[1]=ALerpF1(a[1],b[1],c[1]);return d;}
+ A_STATIC retAF3 opALerpF3(outAF3 d,inAF3 a,inAF3 b,inAF3 c){d[0]=ALerpF1(a[0],b[0],c[0]);d[1]=ALerpF1(a[1],b[1],c[1]);d[2]=ALerpF1(a[2],b[2],c[2]);return d;}
+ A_STATIC retAF4 opALerpF4(outAF4 d,inAF4 a,inAF4 b,inAF4 c){d[0]=ALerpF1(a[0],b[0],c[0]);d[1]=ALerpF1(a[1],b[1],c[1]);d[2]=ALerpF1(a[2],b[2],c[2]);d[3]=ALerpF1(a[3],b[3],c[3]);return d;}
+//==============================================================================================================================
+ A_STATIC retAD2 opALerpOneD2(outAD2 d,inAD2 a,inAD2 b,AD1 c){d[0]=ALerpD1(a[0],b[0],c);d[1]=ALerpD1(a[1],b[1],c);return d;}
+ A_STATIC retAD3 opALerpOneD3(outAD3 d,inAD3 a,inAD3 b,AD1 c){d[0]=ALerpD1(a[0],b[0],c);d[1]=ALerpD1(a[1],b[1],c);d[2]=ALerpD1(a[2],b[2],c);return d;}
+ A_STATIC retAD4 opALerpOneD4(outAD4 d,inAD4 a,inAD4 b,AD1 c){d[0]=ALerpD1(a[0],b[0],c);d[1]=ALerpD1(a[1],b[1],c);d[2]=ALerpD1(a[2],b[2],c);d[3]=ALerpD1(a[3],b[3],c);return d;}
+//------------------------------------------------------------------------------------------------------------------------------
+ A_STATIC retAF2 opALerpOneF2(outAF2 d,inAF2 a,inAF2 b,AF1 c){d[0]=ALerpF1(a[0],b[0],c);d[1]=ALerpF1(a[1],b[1],c);return d;}
+ A_STATIC retAF3 opALerpOneF3(outAF3 d,inAF3 a,inAF3 b,AF1 c){d[0]=ALerpF1(a[0],b[0],c);d[1]=ALerpF1(a[1],b[1],c);d[2]=ALerpF1(a[2],b[2],c);return d;}
+ A_STATIC retAF4 opALerpOneF4(outAF4 d,inAF4 a,inAF4 b,AF1 c){d[0]=ALerpF1(a[0],b[0],c);d[1]=ALerpF1(a[1],b[1],c);d[2]=ALerpF1(a[2],b[2],c);d[3]=ALerpF1(a[3],b[3],c);return d;}
+//==============================================================================================================================
+ A_STATIC retAD2 opAMaxD2(outAD2 d,inAD2 a,inAD2 b){d[0]=AMaxD1(a[0],b[0]);d[1]=AMaxD1(a[1],b[1]);return d;}
+ A_STATIC retAD3 opAMaxD3(outAD3 d,inAD3 a,inAD3 b){d[0]=AMaxD1(a[0],b[0]);d[1]=AMaxD1(a[1],b[1]);d[2]=AMaxD1(a[2],b[2]);return d;}
+ A_STATIC retAD4 opAMaxD4(outAD4 d,inAD4 a,inAD4 b){d[0]=AMaxD1(a[0],b[0]);d[1]=AMaxD1(a[1],b[1]);d[2]=AMaxD1(a[2],b[2]);d[3]=AMaxD1(a[3],b[3]);return d;}
+//------------------------------------------------------------------------------------------------------------------------------
+ A_STATIC retAF2 opAMaxF2(outAF2 d,inAF2 a,inAF2 b){d[0]=AMaxF1(a[0],b[0]);d[1]=AMaxF1(a[1],b[1]);return d;}
+ A_STATIC retAF3 opAMaxF3(outAF3 d,inAF3 a,inAF3 b){d[0]=AMaxF1(a[0],b[0]);d[1]=AMaxF1(a[1],b[1]);d[2]=AMaxF1(a[2],b[2]);return d;}
+ A_STATIC retAF4 opAMaxF4(outAF4 d,inAF4 a,inAF4 b){d[0]=AMaxF1(a[0],b[0]);d[1]=AMaxF1(a[1],b[1]);d[2]=AMaxF1(a[2],b[2]);d[3]=AMaxF1(a[3],b[3]);return d;}
+//==============================================================================================================================
+ A_STATIC retAD2 opAMinD2(outAD2 d,inAD2 a,inAD2 b){d[0]=AMinD1(a[0],b[0]);d[1]=AMinD1(a[1],b[1]);return d;}
+ A_STATIC retAD3 opAMinD3(outAD3 d,inAD3 a,inAD3 b){d[0]=AMinD1(a[0],b[0]);d[1]=AMinD1(a[1],b[1]);d[2]=AMinD1(a[2],b[2]);return d;}
+ A_STATIC retAD4 opAMinD4(outAD4 d,inAD4 a,inAD4 b){d[0]=AMinD1(a[0],b[0]);d[1]=AMinD1(a[1],b[1]);d[2]=AMinD1(a[2],b[2]);d[3]=AMinD1(a[3],b[3]);return d;}
+//------------------------------------------------------------------------------------------------------------------------------
+ A_STATIC retAF2 opAMinF2(outAF2 d,inAF2 a,inAF2 b){d[0]=AMinF1(a[0],b[0]);d[1]=AMinF1(a[1],b[1]);return d;}
+ A_STATIC retAF3 opAMinF3(outAF3 d,inAF3 a,inAF3 b){d[0]=AMinF1(a[0],b[0]);d[1]=AMinF1(a[1],b[1]);d[2]=AMinF1(a[2],b[2]);return d;}
+ A_STATIC retAF4 opAMinF4(outAF4 d,inAF4 a,inAF4 b){d[0]=AMinF1(a[0],b[0]);d[1]=AMinF1(a[1],b[1]);d[2]=AMinF1(a[2],b[2]);d[3]=AMinF1(a[3],b[3]);return d;}
+//==============================================================================================================================
+ A_STATIC retAD2 opAMulD2(outAD2 d,inAD2 a,inAD2 b){d[0]=a[0]*b[0];d[1]=a[1]*b[1];return d;}
+ A_STATIC retAD3 opAMulD3(outAD3 d,inAD3 a,inAD3 b){d[0]=a[0]*b[0];d[1]=a[1]*b[1];d[2]=a[2]*b[2];return d;}
+ A_STATIC retAD4 opAMulD4(outAD4 d,inAD4 a,inAD4 b){d[0]=a[0]*b[0];d[1]=a[1]*b[1];d[2]=a[2]*b[2];d[3]=a[3]*b[3];return d;}
+//------------------------------------------------------------------------------------------------------------------------------
+ A_STATIC retAF2 opAMulF2(outAF2 d,inAF2 a,inAF2 b){d[0]=a[0]*b[0];d[1]=a[1]*b[1];return d;}
+ A_STATIC retAF3 opAMulF3(outAF3 d,inAF3 a,inAF3 b){d[0]=a[0]*b[0];d[1]=a[1]*b[1];d[2]=a[2]*b[2];return d;}
+ A_STATIC retAF4 opAMulF4(outAF4 d,inAF4 a,inAF4 b){d[0]=a[0]*b[0];d[1]=a[1]*b[1];d[2]=a[2]*b[2];d[3]=a[3]*b[3];return d;}
+//==============================================================================================================================
+ A_STATIC retAD2 opAMulOneD2(outAD2 d,inAD2 a,AD1 b){d[0]=a[0]*b;d[1]=a[1]*b;return d;}
+ A_STATIC retAD3 opAMulOneD3(outAD3 d,inAD3 a,AD1 b){d[0]=a[0]*b;d[1]=a[1]*b;d[2]=a[2]*b;return d;}
+ A_STATIC retAD4 opAMulOneD4(outAD4 d,inAD4 a,AD1 b){d[0]=a[0]*b;d[1]=a[1]*b;d[2]=a[2]*b;d[3]=a[3]*b;return d;}
+//------------------------------------------------------------------------------------------------------------------------------
+ A_STATIC retAF2 opAMulOneF2(outAF2 d,inAF2 a,AF1 b){d[0]=a[0]*b;d[1]=a[1]*b;return d;}
+ A_STATIC retAF3 opAMulOneF3(outAF3 d,inAF3 a,AF1 b){d[0]=a[0]*b;d[1]=a[1]*b;d[2]=a[2]*b;return d;}
+ A_STATIC retAF4 opAMulOneF4(outAF4 d,inAF4 a,AF1 b){d[0]=a[0]*b;d[1]=a[1]*b;d[2]=a[2]*b;d[3]=a[3]*b;return d;}
+//==============================================================================================================================
+ A_STATIC retAD2 opANegD2(outAD2 d,inAD2 a){d[0]=-a[0];d[1]=-a[1];return d;}
+ A_STATIC retAD3 opANegD3(outAD3 d,inAD3 a){d[0]=-a[0];d[1]=-a[1];d[2]=-a[2];return d;}
+ A_STATIC retAD4 opANegD4(outAD4 d,inAD4 a){d[0]=-a[0];d[1]=-a[1];d[2]=-a[2];d[3]=-a[3];return d;}
+//------------------------------------------------------------------------------------------------------------------------------
+ A_STATIC retAF2 opANegF2(outAF2 d,inAF2 a){d[0]=-a[0];d[1]=-a[1];return d;}
+ A_STATIC retAF3 opANegF3(outAF3 d,inAF3 a){d[0]=-a[0];d[1]=-a[1];d[2]=-a[2];return d;}
+ A_STATIC retAF4 opANegF4(outAF4 d,inAF4 a){d[0]=-a[0];d[1]=-a[1];d[2]=-a[2];d[3]=-a[3];return d;}
+//==============================================================================================================================
+ A_STATIC retAD2 opARcpD2(outAD2 d,inAD2 a){d[0]=ARcpD1(a[0]);d[1]=ARcpD1(a[1]);return d;}
+ A_STATIC retAD3 opARcpD3(outAD3 d,inAD3 a){d[0]=ARcpD1(a[0]);d[1]=ARcpD1(a[1]);d[2]=ARcpD1(a[2]);return d;}
+ A_STATIC retAD4 opARcpD4(outAD4 d,inAD4 a){d[0]=ARcpD1(a[0]);d[1]=ARcpD1(a[1]);d[2]=ARcpD1(a[2]);d[3]=ARcpD1(a[3]);return d;}
+//------------------------------------------------------------------------------------------------------------------------------
+ A_STATIC retAF2 opARcpF2(outAF2 d,inAF2 a){d[0]=ARcpF1(a[0]);d[1]=ARcpF1(a[1]);return d;}
+ A_STATIC retAF3 opARcpF3(outAF3 d,inAF3 a){d[0]=ARcpF1(a[0]);d[1]=ARcpF1(a[1]);d[2]=ARcpF1(a[2]);return d;}
+ A_STATIC retAF4 opARcpF4(outAF4 d,inAF4 a){d[0]=ARcpF1(a[0]);d[1]=ARcpF1(a[1]);d[2]=ARcpF1(a[2]);d[3]=ARcpF1(a[3]);return d;}
+////////////////////////////////////////////////////////////////////////////////////////////////////////////////////////////////
+////////////////////////////////////////////////////////////////////////////////////////////////////////////////////////////////
+//_____________________________________________________________/\_______________________________________________________________
+//==============================================================================================================================
+//                                                     HALF FLOAT PACKING
+//==============================================================================================================================
+ // Convert float to half (in lower 16-bits of output).
+ // Same fast technique as documented here: ftp://ftp.fox-toolkit.org/pub/fasthalffloatconversion.pdf
+ // Supports denormals.
+ // Conversion rules are to make computations possibly "safer" on the GPU,
+ //  -INF & -NaN -> -65504
+ //  +INF & +NaN -> +65504
+ A_STATIC AU1 AU1_AH1_AF1(AF1 f){
+  static AW1 base[512]={
+   0x0000,0x0000,0x0000,0x0000,0x0000,0x0000,0x0000,0x0000,0x0000,0x0000,0x0000,0x0000,0x0000,0x0000,0x0000,0x0000,
+   0x0000,0x0000,0x0000,0x0000,0x0000,0x0000,0x0000,0x0000,0x0000,0x0000,0x0000,0x0000,0x0000,0x0000,0x0000,0x0000,
+   0x0000,0x0000,0x0000,0x0000,0x0000,0x0000,0x0000,0x0000,0x0000,0x0000,0x0000,0x0000,0x0000,0x0000,0x0000,0x0000,
+   0x0000,0x0000,0x0000,0x0000,0x0000,0x0000,0x0000,0x0000,0x0000,0x0000,0x0000,0x0000,0x0000,0x0000,0x0000,0x0000,
+   0x0000,0x0000,0x0000,0x0000,0x0000,0x0000,0x0000,0x0000,0x0000,0x0000,0x0000,0x0000,0x0000,0x0000,0x0000,0x0000,
+   0x0000,0x0000,0x0000,0x0000,0x0000,0x0000,0x0000,0x0000,0x0000,0x0000,0x0000,0x0000,0x0000,0x0000,0x0000,0x0000,
+   0x0000,0x0000,0x0000,0x0000,0x0000,0x0000,0x0000,0x0001,0x0002,0x0004,0x0008,0x0010,0x0020,0x0040,0x0080,0x0100,
+   0x0200,0x0400,0x0800,0x0c00,0x1000,0x1400,0x1800,0x1c00,0x2000,0x2400,0x2800,0x2c00,0x3000,0x3400,0x3800,0x3c00,
+   0x4000,0x4400,0x4800,0x4c00,0x5000,0x5400,0x5800,0x5c00,0x6000,0x6400,0x6800,0x6c00,0x7000,0x7400,0x7800,0x7bff,
+   0x7bff,0x7bff,0x7bff,0x7bff,0x7bff,0x7bff,0x7bff,0x7bff,0x7bff,0x7bff,0x7bff,0x7bff,0x7bff,0x7bff,0x7bff,0x7bff,
+   0x7bff,0x7bff,0x7bff,0x7bff,0x7bff,0x7bff,0x7bff,0x7bff,0x7bff,0x7bff,0x7bff,0x7bff,0x7bff,0x7bff,0x7bff,0x7bff,
+   0x7bff,0x7bff,0x7bff,0x7bff,0x7bff,0x7bff,0x7bff,0x7bff,0x7bff,0x7bff,0x7bff,0x7bff,0x7bff,0x7bff,0x7bff,0x7bff,
+   0x7bff,0x7bff,0x7bff,0x7bff,0x7bff,0x7bff,0x7bff,0x7bff,0x7bff,0x7bff,0x7bff,0x7bff,0x7bff,0x7bff,0x7bff,0x7bff,
+   0x7bff,0x7bff,0x7bff,0x7bff,0x7bff,0x7bff,0x7bff,0x7bff,0x7bff,0x7bff,0x7bff,0x7bff,0x7bff,0x7bff,0x7bff,0x7bff,
+   0x7bff,0x7bff,0x7bff,0x7bff,0x7bff,0x7bff,0x7bff,0x7bff,0x7bff,0x7bff,0x7bff,0x7bff,0x7bff,0x7bff,0x7bff,0x7bff,
+   0x7bff,0x7bff,0x7bff,0x7bff,0x7bff,0x7bff,0x7bff,0x7bff,0x7bff,0x7bff,0x7bff,0x7bff,0x7bff,0x7bff,0x7bff,0x7bff,
+   0x8000,0x8000,0x8000,0x8000,0x8000,0x8000,0x8000,0x8000,0x8000,0x8000,0x8000,0x8000,0x8000,0x8000,0x8000,0x8000,
+   0x8000,0x8000,0x8000,0x8000,0x8000,0x8000,0x8000,0x8000,0x8000,0x8000,0x8000,0x8000,0x8000,0x8000,0x8000,0x8000,
+   0x8000,0x8000,0x8000,0x8000,0x8000,0x8000,0x8000,0x8000,0x8000,0x8000,0x8000,0x8000,0x8000,0x8000,0x8000,0x8000,
+   0x8000,0x8000,0x8000,0x8000,0x8000,0x8000,0x8000,0x8000,0x8000,0x8000,0x8000,0x8000,0x8000,0x8000,0x8000,0x8000,
+   0x8000,0x8000,0x8000,0x8000,0x8000,0x8000,0x8000,0x8000,0x8000,0x8000,0x8000,0x8000,0x8000,0x8000,0x8000,0x8000,
+   0x8000,0x8000,0x8000,0x8000,0x8000,0x8000,0x8000,0x8000,0x8000,0x8000,0x8000,0x8000,0x8000,0x8000,0x8000,0x8000,
+   0x8000,0x8000,0x8000,0x8000,0x8000,0x8000,0x8000,0x8001,0x8002,0x8004,0x8008,0x8010,0x8020,0x8040,0x8080,0x8100,
+   0x8200,0x8400,0x8800,0x8c00,0x9000,0x9400,0x9800,0x9c00,0xa000,0xa400,0xa800,0xac00,0xb000,0xb400,0xb800,0xbc00,
+   0xc000,0xc400,0xc800,0xcc00,0xd000,0xd400,0xd800,0xdc00,0xe000,0xe400,0xe800,0xec00,0xf000,0xf400,0xf800,0xfbff,
+   0xfbff,0xfbff,0xfbff,0xfbff,0xfbff,0xfbff,0xfbff,0xfbff,0xfbff,0xfbff,0xfbff,0xfbff,0xfbff,0xfbff,0xfbff,0xfbff,
+   0xfbff,0xfbff,0xfbff,0xfbff,0xfbff,0xfbff,0xfbff,0xfbff,0xfbff,0xfbff,0xfbff,0xfbff,0xfbff,0xfbff,0xfbff,0xfbff,
+   0xfbff,0xfbff,0xfbff,0xfbff,0xfbff,0xfbff,0xfbff,0xfbff,0xfbff,0xfbff,0xfbff,0xfbff,0xfbff,0xfbff,0xfbff,0xfbff,
+   0xfbff,0xfbff,0xfbff,0xfbff,0xfbff,0xfbff,0xfbff,0xfbff,0xfbff,0xfbff,0xfbff,0xfbff,0xfbff,0xfbff,0xfbff,0xfbff,
+   0xfbff,0xfbff,0xfbff,0xfbff,0xfbff,0xfbff,0xfbff,0xfbff,0xfbff,0xfbff,0xfbff,0xfbff,0xfbff,0xfbff,0xfbff,0xfbff,
+   0xfbff,0xfbff,0xfbff,0xfbff,0xfbff,0xfbff,0xfbff,0xfbff,0xfbff,0xfbff,0xfbff,0xfbff,0xfbff,0xfbff,0xfbff,0xfbff,
+   0xfbff,0xfbff,0xfbff,0xfbff,0xfbff,0xfbff,0xfbff,0xfbff,0xfbff,0xfbff,0xfbff,0xfbff,0xfbff,0xfbff,0xfbff,0xfbff};
+  static AB1 shift[512]={
+   0x18,0x18,0x18,0x18,0x18,0x18,0x18,0x18,0x18,0x18,0x18,0x18,0x18,0x18,0x18,0x18,
+   0x18,0x18,0x18,0x18,0x18,0x18,0x18,0x18,0x18,0x18,0x18,0x18,0x18,0x18,0x18,0x18,
+   0x18,0x18,0x18,0x18,0x18,0x18,0x18,0x18,0x18,0x18,0x18,0x18,0x18,0x18,0x18,0x18,
+   0x18,0x18,0x18,0x18,0x18,0x18,0x18,0x18,0x18,0x18,0x18,0x18,0x18,0x18,0x18,0x18,
+   0x18,0x18,0x18,0x18,0x18,0x18,0x18,0x18,0x18,0x18,0x18,0x18,0x18,0x18,0x18,0x18,
+   0x18,0x18,0x18,0x18,0x18,0x18,0x18,0x18,0x18,0x18,0x18,0x18,0x18,0x18,0x18,0x18,
+   0x18,0x18,0x18,0x18,0x18,0x18,0x18,0x17,0x16,0x15,0x14,0x13,0x12,0x11,0x10,0x0f,
+   0x0e,0x0d,0x0d,0x0d,0x0d,0x0d,0x0d,0x0d,0x0d,0x0d,0x0d,0x0d,0x0d,0x0d,0x0d,0x0d,
+   0x0d,0x0d,0x0d,0x0d,0x0d,0x0d,0x0d,0x0d,0x0d,0x0d,0x0d,0x0d,0x0d,0x0d,0x0d,0x18,
+   0x18,0x18,0x18,0x18,0x18,0x18,0x18,0x18,0x18,0x18,0x18,0x18,0x18,0x18,0x18,0x18,
+   0x18,0x18,0x18,0x18,0x18,0x18,0x18,0x18,0x18,0x18,0x18,0x18,0x18,0x18,0x18,0x18,
+   0x18,0x18,0x18,0x18,0x18,0x18,0x18,0x18,0x18,0x18,0x18,0x18,0x18,0x18,0x18,0x18,
+   0x18,0x18,0x18,0x18,0x18,0x18,0x18,0x18,0x18,0x18,0x18,0x18,0x18,0x18,0x18,0x18,
+   0x18,0x18,0x18,0x18,0x18,0x18,0x18,0x18,0x18,0x18,0x18,0x18,0x18,0x18,0x18,0x18,
+   0x18,0x18,0x18,0x18,0x18,0x18,0x18,0x18,0x18,0x18,0x18,0x18,0x18,0x18,0x18,0x18,
+   0x18,0x18,0x18,0x18,0x18,0x18,0x18,0x18,0x18,0x18,0x18,0x18,0x18,0x18,0x18,0x18,
+   0x18,0x18,0x18,0x18,0x18,0x18,0x18,0x18,0x18,0x18,0x18,0x18,0x18,0x18,0x18,0x18,
+   0x18,0x18,0x18,0x18,0x18,0x18,0x18,0x18,0x18,0x18,0x18,0x18,0x18,0x18,0x18,0x18,
+   0x18,0x18,0x18,0x18,0x18,0x18,0x18,0x18,0x18,0x18,0x18,0x18,0x18,0x18,0x18,0x18,
+   0x18,0x18,0x18,0x18,0x18,0x18,0x18,0x18,0x18,0x18,0x18,0x18,0x18,0x18,0x18,0x18,
+   0x18,0x18,0x18,0x18,0x18,0x18,0x18,0x18,0x18,0x18,0x18,0x18,0x18,0x18,0x18,0x18,
+   0x18,0x18,0x18,0x18,0x18,0x18,0x18,0x18,0x18,0x18,0x18,0x18,0x18,0x18,0x18,0x18,
+   0x18,0x18,0x18,0x18,0x18,0x18,0x18,0x17,0x16,0x15,0x14,0x13,0x12,0x11,0x10,0x0f,
+   0x0e,0x0d,0x0d,0x0d,0x0d,0x0d,0x0d,0x0d,0x0d,0x0d,0x0d,0x0d,0x0d,0x0d,0x0d,0x0d,
+   0x0d,0x0d,0x0d,0x0d,0x0d,0x0d,0x0d,0x0d,0x0d,0x0d,0x0d,0x0d,0x0d,0x0d,0x0d,0x18,
+   0x18,0x18,0x18,0x18,0x18,0x18,0x18,0x18,0x18,0x18,0x18,0x18,0x18,0x18,0x18,0x18,
+   0x18,0x18,0x18,0x18,0x18,0x18,0x18,0x18,0x18,0x18,0x18,0x18,0x18,0x18,0x18,0x18,
+   0x18,0x18,0x18,0x18,0x18,0x18,0x18,0x18,0x18,0x18,0x18,0x18,0x18,0x18,0x18,0x18,
+   0x18,0x18,0x18,0x18,0x18,0x18,0x18,0x18,0x18,0x18,0x18,0x18,0x18,0x18,0x18,0x18,
+   0x18,0x18,0x18,0x18,0x18,0x18,0x18,0x18,0x18,0x18,0x18,0x18,0x18,0x18,0x18,0x18,
+   0x18,0x18,0x18,0x18,0x18,0x18,0x18,0x18,0x18,0x18,0x18,0x18,0x18,0x18,0x18,0x18,
+   0x18,0x18,0x18,0x18,0x18,0x18,0x18,0x18,0x18,0x18,0x18,0x18,0x18,0x18,0x18,0x18};
+  union{AF1 f;AU1 u;}bits;bits.f=f;AU1 u=bits.u;AU1 i=u>>23;return (AU1)(base[i])+((u&0x7fffff)>>shift[i]);}
+//------------------------------------------------------------------------------------------------------------------------------
+ // Used to output packed constant.
+ A_STATIC AU1 AU1_AH2_AF2(inAF2 a){return AU1_AH1_AF1(a[0])+(AU1_AH1_AF1(a[1])<<16);}
+#endif
+////////////////////////////////////////////////////////////////////////////////////////////////////////////////////////////////
+////////////////////////////////////////////////////////////////////////////////////////////////////////////////////////////////
+////////////////////////////////////////////////////////////////////////////////////////////////////////////////////////////////
+////////////////////////////////////////////////////////////////////////////////////////////////////////////////////////////////
+////////////////////////////////////////////////////////////////////////////////////////////////////////////////////////////////
+////////////////////////////////////////////////////////////////////////////////////////////////////////////////////////////////
+//_____________________________________________________________/\_______________________________________________________________
+//==============================================================================================================================
+//
+//
+//                                                            GLSL
+//
+//
+//==============================================================================================================================
+#if defined(A_GLSL) && defined(A_GPU)
+ #ifndef A_SKIP_EXT
+  #ifdef A_HALF
+   #extension GL_EXT_shader_16bit_storage:require
+   #extension GL_EXT_shader_explicit_arithmetic_types:require 
+  #endif
+//------------------------------------------------------------------------------------------------------------------------------
+  #ifdef A_LONG
+   #extension GL_ARB_gpu_shader_int64:require
+   #extension GL_NV_shader_atomic_int64:require
+  #endif
+//------------------------------------------------------------------------------------------------------------------------------
+  #ifdef A_WAVE
+   #extension GL_KHR_shader_subgroup_arithmetic:require
+   #extension GL_KHR_shader_subgroup_ballot:require
+   #extension GL_KHR_shader_subgroup_quad:require
+   #extension GL_KHR_shader_subgroup_shuffle:require
+  #endif
+ #endif
+//==============================================================================================================================
+ #define AP1 bool
+ #define AP2 bvec2
+ #define AP3 bvec3
+ #define AP4 bvec4
+//------------------------------------------------------------------------------------------------------------------------------
+ #define AF1 float
+ #define AF2 vec2
+ #define AF3 vec3
+ #define AF4 vec4
+//------------------------------------------------------------------------------------------------------------------------------
+ #define AU1 uint
+ #define AU2 uvec2
+ #define AU3 uvec3
+ #define AU4 uvec4
+//------------------------------------------------------------------------------------------------------------------------------
+ #define ASU1 int
+ #define ASU2 ivec2
+ #define ASU3 ivec3
+ #define ASU4 ivec4
+//==============================================================================================================================
+ #define AF1_AU1(x) uintBitsToFloat(AU1(x))
+ #define AF2_AU2(x) uintBitsToFloat(AU2(x))
+ #define AF3_AU3(x) uintBitsToFloat(AU3(x))
+ #define AF4_AU4(x) uintBitsToFloat(AU4(x))
+//------------------------------------------------------------------------------------------------------------------------------
+ #define AU1_AF1(x) floatBitsToUint(AF1(x))
+ #define AU2_AF2(x) floatBitsToUint(AF2(x))
+ #define AU3_AF3(x) floatBitsToUint(AF3(x))
+ #define AU4_AF4(x) floatBitsToUint(AF4(x))
+//------------------------------------------------------------------------------------------------------------------------------
+ AU1 AU1_AH1_AF1_x(AF1 a){return packHalf2x16(AF2(a,0.0));}
+ #define AU1_AH1_AF1(a) AU1_AH1_AF1_x(AF1(a))
+//------------------------------------------------------------------------------------------------------------------------------
+ #define AU1_AH2_AF2 packHalf2x16
+ #define AU1_AW2Unorm_AF2 packUnorm2x16
+ #define AU1_AB4Unorm_AF4 packUnorm4x8
+//------------------------------------------------------------------------------------------------------------------------------
+ #define AF2_AH2_AU1 unpackHalf2x16
+ #define AF2_AW2Unorm_AU1 unpackUnorm2x16
+ #define AF4_AB4Unorm_AU1 unpackUnorm4x8
+//==============================================================================================================================
+ AF1 AF1_x(AF1 a){return AF1(a);}
+ AF2 AF2_x(AF1 a){return AF2(a,a);}
+ AF3 AF3_x(AF1 a){return AF3(a,a,a);}
+ AF4 AF4_x(AF1 a){return AF4(a,a,a,a);}
+ #define AF1_(a) AF1_x(AF1(a))
+ #define AF2_(a) AF2_x(AF1(a))
+ #define AF3_(a) AF3_x(AF1(a))
+ #define AF4_(a) AF4_x(AF1(a))
+//------------------------------------------------------------------------------------------------------------------------------
+ AU1 AU1_x(AU1 a){return AU1(a);}
+ AU2 AU2_x(AU1 a){return AU2(a,a);}
+ AU3 AU3_x(AU1 a){return AU3(a,a,a);}
+ AU4 AU4_x(AU1 a){return AU4(a,a,a,a);}
+ #define AU1_(a) AU1_x(AU1(a))
+ #define AU2_(a) AU2_x(AU1(a))
+ #define AU3_(a) AU3_x(AU1(a))
+ #define AU4_(a) AU4_x(AU1(a))
+//==============================================================================================================================
+ AU1 AAbsSU1(AU1 a){return AU1(abs(ASU1(a)));}
+ AU2 AAbsSU2(AU2 a){return AU2(abs(ASU2(a)));}
+ AU3 AAbsSU3(AU3 a){return AU3(abs(ASU3(a)));}
+ AU4 AAbsSU4(AU4 a){return AU4(abs(ASU4(a)));}
+//------------------------------------------------------------------------------------------------------------------------------
+ AU1 ABfe(AU1 src,AU1 off,AU1 bits){return bitfieldExtract(src,ASU1(off),ASU1(bits));}
+ AU1 ABfi(AU1 src,AU1 ins,AU1 mask){return (ins&mask)|(src&(~mask));}
+ // Proxy for V_BFI_B32 where the 'mask' is set as 'bits', 'mask=(1<<bits)-1', and 'bits' needs to be an immediate.
+ AU1 ABfiM(AU1 src,AU1 ins,AU1 bits){return bitfieldInsert(src,ins,0,ASU1(bits));}
+//------------------------------------------------------------------------------------------------------------------------------
+ // V_MED3_F32.
+ AF1 AClampF1(AF1 x,AF1 n,AF1 m){return clamp(x,n,m);}
+ AF2 AClampF2(AF2 x,AF2 n,AF2 m){return clamp(x,n,m);}
+ AF3 AClampF3(AF3 x,AF3 n,AF3 m){return clamp(x,n,m);}
+ AF4 AClampF4(AF4 x,AF4 n,AF4 m){return clamp(x,n,m);}
+//------------------------------------------------------------------------------------------------------------------------------
+ // V_FRACT_F32 (note DX frac() is different).
+ AF1 AFractF1(AF1 x){return fract(x);}
+ AF2 AFractF2(AF2 x){return fract(x);}
+ AF3 AFractF3(AF3 x){return fract(x);}
+ AF4 AFractF4(AF4 x){return fract(x);}
+//------------------------------------------------------------------------------------------------------------------------------
+ AF1 ALerpF1(AF1 x,AF1 y,AF1 a){return mix(x,y,a);}
+ AF2 ALerpF2(AF2 x,AF2 y,AF2 a){return mix(x,y,a);}
+ AF3 ALerpF3(AF3 x,AF3 y,AF3 a){return mix(x,y,a);}
+ AF4 ALerpF4(AF4 x,AF4 y,AF4 a){return mix(x,y,a);}
+//------------------------------------------------------------------------------------------------------------------------------
+ // V_MAX3_F32.
+ AF1 AMax3F1(AF1 x,AF1 y,AF1 z){return max(x,max(y,z));}
+ AF2 AMax3F2(AF2 x,AF2 y,AF2 z){return max(x,max(y,z));}
+ AF3 AMax3F3(AF3 x,AF3 y,AF3 z){return max(x,max(y,z));}
+ AF4 AMax3F4(AF4 x,AF4 y,AF4 z){return max(x,max(y,z));}
+//------------------------------------------------------------------------------------------------------------------------------
+ AU1 AMax3SU1(AU1 x,AU1 y,AU1 z){return AU1(max(ASU1(x),max(ASU1(y),ASU1(z))));}
+ AU2 AMax3SU2(AU2 x,AU2 y,AU2 z){return AU2(max(ASU2(x),max(ASU2(y),ASU2(z))));}
+ AU3 AMax3SU3(AU3 x,AU3 y,AU3 z){return AU3(max(ASU3(x),max(ASU3(y),ASU3(z))));}
+ AU4 AMax3SU4(AU4 x,AU4 y,AU4 z){return AU4(max(ASU4(x),max(ASU4(y),ASU4(z))));}
+//------------------------------------------------------------------------------------------------------------------------------
+ AU1 AMax3U1(AU1 x,AU1 y,AU1 z){return max(x,max(y,z));}
+ AU2 AMax3U2(AU2 x,AU2 y,AU2 z){return max(x,max(y,z));}
+ AU3 AMax3U3(AU3 x,AU3 y,AU3 z){return max(x,max(y,z));}
+ AU4 AMax3U4(AU4 x,AU4 y,AU4 z){return max(x,max(y,z));}
+//------------------------------------------------------------------------------------------------------------------------------
+ AU1 AMaxSU1(AU1 a,AU1 b){return AU1(max(ASU1(a),ASU1(b)));}
+ AU2 AMaxSU2(AU2 a,AU2 b){return AU2(max(ASU2(a),ASU2(b)));}
+ AU3 AMaxSU3(AU3 a,AU3 b){return AU3(max(ASU3(a),ASU3(b)));}
+ AU4 AMaxSU4(AU4 a,AU4 b){return AU4(max(ASU4(a),ASU4(b)));}
+//------------------------------------------------------------------------------------------------------------------------------
+ // Clamp has an easier pattern match for med3 when some ordering is known.
+ // V_MED3_F32.
+ AF1 AMed3F1(AF1 x,AF1 y,AF1 z){return max(min(x,y),min(max(x,y),z));}
+ AF2 AMed3F2(AF2 x,AF2 y,AF2 z){return max(min(x,y),min(max(x,y),z));}
+ AF3 AMed3F3(AF3 x,AF3 y,AF3 z){return max(min(x,y),min(max(x,y),z));}
+ AF4 AMed3F4(AF4 x,AF4 y,AF4 z){return max(min(x,y),min(max(x,y),z));}
+//------------------------------------------------------------------------------------------------------------------------------
+ // V_MIN3_F32.
+ AF1 AMin3F1(AF1 x,AF1 y,AF1 z){return min(x,min(y,z));}
+ AF2 AMin3F2(AF2 x,AF2 y,AF2 z){return min(x,min(y,z));}
+ AF3 AMin3F3(AF3 x,AF3 y,AF3 z){return min(x,min(y,z));}
+ AF4 AMin3F4(AF4 x,AF4 y,AF4 z){return min(x,min(y,z));}
+//------------------------------------------------------------------------------------------------------------------------------
+ AU1 AMin3SU1(AU1 x,AU1 y,AU1 z){return AU1(min(ASU1(x),min(ASU1(y),ASU1(z))));}
+ AU2 AMin3SU2(AU2 x,AU2 y,AU2 z){return AU2(min(ASU2(x),min(ASU2(y),ASU2(z))));}
+ AU3 AMin3SU3(AU3 x,AU3 y,AU3 z){return AU3(min(ASU3(x),min(ASU3(y),ASU3(z))));}
+ AU4 AMin3SU4(AU4 x,AU4 y,AU4 z){return AU4(min(ASU4(x),min(ASU4(y),ASU4(z))));}
+//------------------------------------------------------------------------------------------------------------------------------
+ AU1 AMin3U1(AU1 x,AU1 y,AU1 z){return min(x,min(y,z));}
+ AU2 AMin3U2(AU2 x,AU2 y,AU2 z){return min(x,min(y,z));}
+ AU3 AMin3U3(AU3 x,AU3 y,AU3 z){return min(x,min(y,z));}
+ AU4 AMin3U4(AU4 x,AU4 y,AU4 z){return min(x,min(y,z));}
+//------------------------------------------------------------------------------------------------------------------------------
+ AU1 AMinSU1(AU1 a,AU1 b){return AU1(min(ASU1(a),ASU1(b)));}
+ AU2 AMinSU2(AU2 a,AU2 b){return AU2(min(ASU2(a),ASU2(b)));}
+ AU3 AMinSU3(AU3 a,AU3 b){return AU3(min(ASU3(a),ASU3(b)));}
+ AU4 AMinSU4(AU4 a,AU4 b){return AU4(min(ASU4(a),ASU4(b)));}
+//------------------------------------------------------------------------------------------------------------------------------
+ // Normalized trig. Valid input domain is {-256 to +256}. No GLSL compiler intrinsic exists to map to this currently.
+ // V_COS_F32.
+ AF1 ANCosF1(AF1 x){return cos(x*AF1_(A_2PI));}
+ AF2 ANCosF2(AF2 x){return cos(x*AF2_(A_2PI));}
+ AF3 ANCosF3(AF3 x){return cos(x*AF3_(A_2PI));}
+ AF4 ANCosF4(AF4 x){return cos(x*AF4_(A_2PI));}
+//------------------------------------------------------------------------------------------------------------------------------
+ // Normalized trig. Valid input domain is {-256 to +256}. No GLSL compiler intrinsic exists to map to this currently.
+ // V_SIN_F32.
+ AF1 ANSinF1(AF1 x){return sin(x*AF1_(A_2PI));}
+ AF2 ANSinF2(AF2 x){return sin(x*AF2_(A_2PI));}
+ AF3 ANSinF3(AF3 x){return sin(x*AF3_(A_2PI));}
+ AF4 ANSinF4(AF4 x){return sin(x*AF4_(A_2PI));}
+//------------------------------------------------------------------------------------------------------------------------------
+ AF1 ARcpF1(AF1 x){return AF1_(1.0)/x;}
+ AF2 ARcpF2(AF2 x){return AF2_(1.0)/x;}
+ AF3 ARcpF3(AF3 x){return AF3_(1.0)/x;}
+ AF4 ARcpF4(AF4 x){return AF4_(1.0)/x;}
+//------------------------------------------------------------------------------------------------------------------------------
+ AF1 ARsqF1(AF1 x){return AF1_(1.0)/sqrt(x);}
+ AF2 ARsqF2(AF2 x){return AF2_(1.0)/sqrt(x);}
+ AF3 ARsqF3(AF3 x){return AF3_(1.0)/sqrt(x);}
+ AF4 ARsqF4(AF4 x){return AF4_(1.0)/sqrt(x);}
+//------------------------------------------------------------------------------------------------------------------------------
+ AF1 ASatF1(AF1 x){return clamp(x,AF1_(0.0),AF1_(1.0));}
+ AF2 ASatF2(AF2 x){return clamp(x,AF2_(0.0),AF2_(1.0));}
+ AF3 ASatF3(AF3 x){return clamp(x,AF3_(0.0),AF3_(1.0));}
+ AF4 ASatF4(AF4 x){return clamp(x,AF4_(0.0),AF4_(1.0));}
+//------------------------------------------------------------------------------------------------------------------------------
+ AU1 AShrSU1(AU1 a,AU1 b){return AU1(ASU1(a)>>ASU1(b));}
+ AU2 AShrSU2(AU2 a,AU2 b){return AU2(ASU2(a)>>ASU2(b));}
+ AU3 AShrSU3(AU3 a,AU3 b){return AU3(ASU3(a)>>ASU3(b));}
+ AU4 AShrSU4(AU4 a,AU4 b){return AU4(ASU4(a)>>ASU4(b));}
+////////////////////////////////////////////////////////////////////////////////////////////////////////////////////////////////
+////////////////////////////////////////////////////////////////////////////////////////////////////////////////////////////////
+//_____________________________________________________________/\_______________________________________________________________
+//==============================================================================================================================
+//                                                          GLSL BYTE
+//==============================================================================================================================
+ #ifdef A_BYTE
+  #define AB1 uint8_t
+  #define AB2 u8vec2
+  #define AB3 u8vec3
+  #define AB4 u8vec4
+//------------------------------------------------------------------------------------------------------------------------------
+  #define ASB1 int8_t
+  #define ASB2 i8vec2
+  #define ASB3 i8vec3
+  #define ASB4 i8vec4
+//------------------------------------------------------------------------------------------------------------------------------
+  AB1 AB1_x(AB1 a){return AB1(a);}
+  AB2 AB2_x(AB1 a){return AB2(a,a);}
+  AB3 AB3_x(AB1 a){return AB3(a,a,a);}
+  AB4 AB4_x(AB1 a){return AB4(a,a,a,a);}
+  #define AB1_(a) AB1_x(AB1(a))
+  #define AB2_(a) AB2_x(AB1(a))
+  #define AB3_(a) AB3_x(AB1(a))
+  #define AB4_(a) AB4_x(AB1(a))
+ #endif
+////////////////////////////////////////////////////////////////////////////////////////////////////////////////////////////////
+////////////////////////////////////////////////////////////////////////////////////////////////////////////////////////////////
+//_____________________________________________________________/\_______________________________________________________________
+//==============================================================================================================================
+//                                                          GLSL HALF
+//==============================================================================================================================
+ #ifdef A_HALF
+  #define AH1 float16_t
+  #define AH2 f16vec2
+  #define AH3 f16vec3
+  #define AH4 f16vec4
+//------------------------------------------------------------------------------------------------------------------------------
+  #define AW1 uint16_t
+  #define AW2 u16vec2
+  #define AW3 u16vec3
+  #define AW4 u16vec4
+//------------------------------------------------------------------------------------------------------------------------------
+  #define ASW1 int16_t
+  #define ASW2 i16vec2
+  #define ASW3 i16vec3
+  #define ASW4 i16vec4
+//==============================================================================================================================
+  #define AH2_AU1(x) unpackFloat2x16(AU1(x))
+  AH4 AH4_AU2_x(AU2 x){return AH4(unpackFloat2x16(x.x),unpackFloat2x16(x.y));}
+  #define AH4_AU2(x) AH4_AU2_x(AU2(x))
+  #define AW2_AU1(x) unpackUint2x16(AU1(x))
+  #define AW4_AU2(x) unpackUint4x16(pack64(AU2(x)))
+//------------------------------------------------------------------------------------------------------------------------------
+  #define AU1_AH2(x) packFloat2x16(AH2(x))
+  AU2 AU2_AH4_x(AH4 x){return AU2(packFloat2x16(x.xy),packFloat2x16(x.zw));}
+  #define AU2_AH4(x) AU2_AH4_x(AH4(x))
+  #define AU1_AW2(x) packUint2x16(AW2(x))
+  #define AU2_AW4(x) unpack32(packUint4x16(AW4(x)))
+//==============================================================================================================================
+  #define AW1_AH1(x) halfBitsToUint16(AH1(x))
+  #define AW2_AH2(x) halfBitsToUint16(AH2(x))
+  #define AW3_AH3(x) halfBitsToUint16(AH3(x))
+  #define AW4_AH4(x) halfBitsToUint16(AH4(x))
+//------------------------------------------------------------------------------------------------------------------------------
+  #define AH1_AW1(x) uint16BitsToHalf(AW1(x))
+  #define AH2_AW2(x) uint16BitsToHalf(AW2(x))
+  #define AH3_AW3(x) uint16BitsToHalf(AW3(x))
+  #define AH4_AW4(x) uint16BitsToHalf(AW4(x))
+//==============================================================================================================================
+  AH1 AH1_x(AH1 a){return AH1(a);}
+  AH2 AH2_x(AH1 a){return AH2(a,a);}
+  AH3 AH3_x(AH1 a){return AH3(a,a,a);}
+  AH4 AH4_x(AH1 a){return AH4(a,a,a,a);}
+  #define AH1_(a) AH1_x(AH1(a))
+  #define AH2_(a) AH2_x(AH1(a))
+  #define AH3_(a) AH3_x(AH1(a))
+  #define AH4_(a) AH4_x(AH1(a))
+//------------------------------------------------------------------------------------------------------------------------------
+  AW1 AW1_x(AW1 a){return AW1(a);}
+  AW2 AW2_x(AW1 a){return AW2(a,a);}
+  AW3 AW3_x(AW1 a){return AW3(a,a,a);}
+  AW4 AW4_x(AW1 a){return AW4(a,a,a,a);}
+  #define AW1_(a) AW1_x(AW1(a))
+  #define AW2_(a) AW2_x(AW1(a))
+  #define AW3_(a) AW3_x(AW1(a))
+  #define AW4_(a) AW4_x(AW1(a))
+//==============================================================================================================================
+  AW1 AAbsSW1(AW1 a){return AW1(abs(ASW1(a)));}
+  AW2 AAbsSW2(AW2 a){return AW2(abs(ASW2(a)));}
+  AW3 AAbsSW3(AW3 a){return AW3(abs(ASW3(a)));}
+  AW4 AAbsSW4(AW4 a){return AW4(abs(ASW4(a)));}
+//------------------------------------------------------------------------------------------------------------------------------
+  AH1 AClampH1(AH1 x,AH1 n,AH1 m){return clamp(x,n,m);}
+  AH2 AClampH2(AH2 x,AH2 n,AH2 m){return clamp(x,n,m);}
+  AH3 AClampH3(AH3 x,AH3 n,AH3 m){return clamp(x,n,m);}
+  AH4 AClampH4(AH4 x,AH4 n,AH4 m){return clamp(x,n,m);}
+//------------------------------------------------------------------------------------------------------------------------------
+  AH1 AFractH1(AH1 x){return fract(x);}
+  AH2 AFractH2(AH2 x){return fract(x);}
+  AH3 AFractH3(AH3 x){return fract(x);}
+  AH4 AFractH4(AH4 x){return fract(x);}
+//------------------------------------------------------------------------------------------------------------------------------
+  AH1 ALerpH1(AH1 x,AH1 y,AH1 a){return mix(x,y,a);}
+  AH2 ALerpH2(AH2 x,AH2 y,AH2 a){return mix(x,y,a);}
+  AH3 ALerpH3(AH3 x,AH3 y,AH3 a){return mix(x,y,a);}
+  AH4 ALerpH4(AH4 x,AH4 y,AH4 a){return mix(x,y,a);}
+//------------------------------------------------------------------------------------------------------------------------------
+  // No packed version of max3.
+  AH1 AMax3H1(AH1 x,AH1 y,AH1 z){return max(x,max(y,z));}
+  AH2 AMax3H2(AH2 x,AH2 y,AH2 z){return max(x,max(y,z));}
+  AH3 AMax3H3(AH3 x,AH3 y,AH3 z){return max(x,max(y,z));}
+  AH4 AMax3H4(AH4 x,AH4 y,AH4 z){return max(x,max(y,z));}
+//------------------------------------------------------------------------------------------------------------------------------
+  AW1 AMaxSW1(AW1 a,AW1 b){return AW1(max(ASU1(a),ASU1(b)));}
+  AW2 AMaxSW2(AW2 a,AW2 b){return AW2(max(ASU2(a),ASU2(b)));}
+  AW3 AMaxSW3(AW3 a,AW3 b){return AW3(max(ASU3(a),ASU3(b)));}
+  AW4 AMaxSW4(AW4 a,AW4 b){return AW4(max(ASU4(a),ASU4(b)));}
+//------------------------------------------------------------------------------------------------------------------------------
+  // No packed version of min3.
+  AH1 AMin3H1(AH1 x,AH1 y,AH1 z){return min(x,min(y,z));}
+  AH2 AMin3H2(AH2 x,AH2 y,AH2 z){return min(x,min(y,z));}
+  AH3 AMin3H3(AH3 x,AH3 y,AH3 z){return min(x,min(y,z));}
+  AH4 AMin3H4(AH4 x,AH4 y,AH4 z){return min(x,min(y,z));}
+//------------------------------------------------------------------------------------------------------------------------------
+  AW1 AMinSW1(AW1 a,AW1 b){return AW1(min(ASU1(a),ASU1(b)));}
+  AW2 AMinSW2(AW2 a,AW2 b){return AW2(min(ASU2(a),ASU2(b)));}
+  AW3 AMinSW3(AW3 a,AW3 b){return AW3(min(ASU3(a),ASU3(b)));}
+  AW4 AMinSW4(AW4 a,AW4 b){return AW4(min(ASU4(a),ASU4(b)));}
+//------------------------------------------------------------------------------------------------------------------------------
+  AH1 ARcpH1(AH1 x){return AH1_(1.0)/x;}
+  AH2 ARcpH2(AH2 x){return AH2_(1.0)/x;}
+  AH3 ARcpH3(AH3 x){return AH3_(1.0)/x;}
+  AH4 ARcpH4(AH4 x){return AH4_(1.0)/x;}
+//------------------------------------------------------------------------------------------------------------------------------
+  AH1 ARsqH1(AH1 x){return AH1_(1.0)/sqrt(x);}
+  AH2 ARsqH2(AH2 x){return AH2_(1.0)/sqrt(x);}
+  AH3 ARsqH3(AH3 x){return AH3_(1.0)/sqrt(x);}
+  AH4 ARsqH4(AH4 x){return AH4_(1.0)/sqrt(x);}
+//------------------------------------------------------------------------------------------------------------------------------
+  AH1 ASatH1(AH1 x){return clamp(x,AH1_(0.0),AH1_(1.0));}
+  AH2 ASatH2(AH2 x){return clamp(x,AH2_(0.0),AH2_(1.0));}
+  AH3 ASatH3(AH3 x){return clamp(x,AH3_(0.0),AH3_(1.0));}
+  AH4 ASatH4(AH4 x){return clamp(x,AH4_(0.0),AH4_(1.0));}
+//------------------------------------------------------------------------------------------------------------------------------
+  AW1 AShrSW1(AW1 a,AW1 b){return AW1(ASW1(a)>>ASW1(b));}
+  AW2 AShrSW2(AW2 a,AW2 b){return AW2(ASW2(a)>>ASW2(b));}
+  AW3 AShrSW3(AW3 a,AW3 b){return AW3(ASW3(a)>>ASW3(b));}
+  AW4 AShrSW4(AW4 a,AW4 b){return AW4(ASW4(a)>>ASW4(b));}
+ #endif
+////////////////////////////////////////////////////////////////////////////////////////////////////////////////////////////////
+////////////////////////////////////////////////////////////////////////////////////////////////////////////////////////////////
+//_____________________________________________________________/\_______________________________________________________________
+//==============================================================================================================================
+//                                                         GLSL DOUBLE
+//==============================================================================================================================
+ #ifdef A_DUBL
+  #define AD1 double
+  #define AD2 dvec2
+  #define AD3 dvec3
+  #define AD4 dvec4
+//------------------------------------------------------------------------------------------------------------------------------
+  AD1 AD1_x(AD1 a){return AD1(a);}
+  AD2 AD2_x(AD1 a){return AD2(a,a);}
+  AD3 AD3_x(AD1 a){return AD3(a,a,a);}
+  AD4 AD4_x(AD1 a){return AD4(a,a,a,a);}
+  #define AD1_(a) AD1_x(AD1(a))
+  #define AD2_(a) AD2_x(AD1(a))
+  #define AD3_(a) AD3_x(AD1(a))
+  #define AD4_(a) AD4_x(AD1(a))
+//==============================================================================================================================
+  AD1 AFractD1(AD1 x){return fract(x);}
+  AD2 AFractD2(AD2 x){return fract(x);}
+  AD3 AFractD3(AD3 x){return fract(x);}
+  AD4 AFractD4(AD4 x){return fract(x);}
+//------------------------------------------------------------------------------------------------------------------------------
+  AD1 ALerpD1(AD1 x,AD1 y,AD1 a){return mix(x,y,a);}
+  AD2 ALerpD2(AD2 x,AD2 y,AD2 a){return mix(x,y,a);}
+  AD3 ALerpD3(AD3 x,AD3 y,AD3 a){return mix(x,y,a);}
+  AD4 ALerpD4(AD4 x,AD4 y,AD4 a){return mix(x,y,a);}
+//------------------------------------------------------------------------------------------------------------------------------
+  AD1 ARcpD1(AD1 x){return AD1_(1.0)/x;}
+  AD2 ARcpD2(AD2 x){return AD2_(1.0)/x;}
+  AD3 ARcpD3(AD3 x){return AD3_(1.0)/x;}
+  AD4 ARcpD4(AD4 x){return AD4_(1.0)/x;}
+//------------------------------------------------------------------------------------------------------------------------------
+  AD1 ARsqD1(AD1 x){return AD1_(1.0)/sqrt(x);}
+  AD2 ARsqD2(AD2 x){return AD2_(1.0)/sqrt(x);}
+  AD3 ARsqD3(AD3 x){return AD3_(1.0)/sqrt(x);}
+  AD4 ARsqD4(AD4 x){return AD4_(1.0)/sqrt(x);}
+//------------------------------------------------------------------------------------------------------------------------------
+  AD1 ASatD1(AD1 x){return clamp(x,AD1_(0.0),AD1_(1.0));}
+  AD2 ASatD2(AD2 x){return clamp(x,AD2_(0.0),AD2_(1.0));}
+  AD3 ASatD3(AD3 x){return clamp(x,AD3_(0.0),AD3_(1.0));}
+  AD4 ASatD4(AD4 x){return clamp(x,AD4_(0.0),AD4_(1.0));}
+ #endif
+////////////////////////////////////////////////////////////////////////////////////////////////////////////////////////////////
+////////////////////////////////////////////////////////////////////////////////////////////////////////////////////////////////
+//_____________________________________________________________/\_______________________________________________________________
+//==============================================================================================================================
+//                                                         GLSL LONG
+//==============================================================================================================================
+ #ifdef A_LONG
+  #define AL1 uint64_t
+  #define AL2 u64vec2
+  #define AL3 u64vec3
+  #define AL4 u64vec4
+//------------------------------------------------------------------------------------------------------------------------------
+  #define ASL1 int64_t
+  #define ASL2 i64vec2
+  #define ASL3 i64vec3
+  #define ASL4 i64vec4
+//------------------------------------------------------------------------------------------------------------------------------
+  #define AL1_AU2(x) packUint2x32(AU2(x))
+  #define AU2_AL1(x) unpackUint2x32(AL1(x))
+//------------------------------------------------------------------------------------------------------------------------------
+  AL1 AL1_x(AL1 a){return AL1(a);}
+  AL2 AL2_x(AL1 a){return AL2(a,a);}
+  AL3 AL3_x(AL1 a){return AL3(a,a,a);}
+  AL4 AL4_x(AL1 a){return AL4(a,a,a,a);}
+  #define AL1_(a) AL1_x(AL1(a))
+  #define AL2_(a) AL2_x(AL1(a))
+  #define AL3_(a) AL3_x(AL1(a))
+  #define AL4_(a) AL4_x(AL1(a))
+//==============================================================================================================================
+  AL1 AAbsSL1(AL1 a){return AL1(abs(ASL1(a)));}
+  AL2 AAbsSL2(AL2 a){return AL2(abs(ASL2(a)));}
+  AL3 AAbsSL3(AL3 a){return AL3(abs(ASL3(a)));}
+  AL4 AAbsSL4(AL4 a){return AL4(abs(ASL4(a)));}
+//------------------------------------------------------------------------------------------------------------------------------
+  AL1 AMaxSL1(AL1 a,AL1 b){return AL1(max(ASU1(a),ASU1(b)));}
+  AL2 AMaxSL2(AL2 a,AL2 b){return AL2(max(ASU2(a),ASU2(b)));}
+  AL3 AMaxSL3(AL3 a,AL3 b){return AL3(max(ASU3(a),ASU3(b)));}
+  AL4 AMaxSL4(AL4 a,AL4 b){return AL4(max(ASU4(a),ASU4(b)));}
+//------------------------------------------------------------------------------------------------------------------------------
+  AL1 AMinSL1(AL1 a,AL1 b){return AL1(min(ASU1(a),ASU1(b)));}
+  AL2 AMinSL2(AL2 a,AL2 b){return AL2(min(ASU2(a),ASU2(b)));}
+  AL3 AMinSL3(AL3 a,AL3 b){return AL3(min(ASU3(a),ASU3(b)));}
+  AL4 AMinSL4(AL4 a,AL4 b){return AL4(min(ASU4(a),ASU4(b)));}
+ #endif
+////////////////////////////////////////////////////////////////////////////////////////////////////////////////////////////////
+////////////////////////////////////////////////////////////////////////////////////////////////////////////////////////////////
+//_____________________________________________________________/\_______________________________________________________________
+//==============================================================================================================================
+//                                                      WAVE OPERATIONS
+//==============================================================================================================================
+ #ifdef A_WAVE
+  // Where 'x' must be a compile time literal.
+  AF1 AWaveXorF1(AF1 v,AU1 x){return subgroupShuffleXor(v,x);}
+  AF2 AWaveXorF2(AF2 v,AU1 x){return subgroupShuffleXor(v,x);}
+  AF3 AWaveXorF3(AF3 v,AU1 x){return subgroupShuffleXor(v,x);}
+  AF4 AWaveXorF4(AF4 v,AU1 x){return subgroupShuffleXor(v,x);}
+  AU1 AWaveXorU1(AU1 v,AU1 x){return subgroupShuffleXor(v,x);}
+  AU2 AWaveXorU2(AU2 v,AU1 x){return subgroupShuffleXor(v,x);}
+  AU3 AWaveXorU3(AU3 v,AU1 x){return subgroupShuffleXor(v,x);}
+  AU4 AWaveXorU4(AU4 v,AU1 x){return subgroupShuffleXor(v,x);}
+//------------------------------------------------------------------------------------------------------------------------------
+  #ifdef A_HALF
+   AH2 AWaveXorH2(AH2 v,AU1 x){return AH2_AU1(subgroupShuffleXor(AU1_AH2(v),x));}
+   AH4 AWaveXorH4(AH4 v,AU1 x){return AH4_AU2(subgroupShuffleXor(AU2_AH4(v),x));}
+   AW2 AWaveXorW2(AW2 v,AU1 x){return AW2_AU1(subgroupShuffleXor(AU1_AW2(v),x));}
+   AW4 AWaveXorW4(AW4 v,AU1 x){return AW4_AU2(subgroupShuffleXor(AU2_AW4(v),x));}
+  #endif
+ #endif
+//==============================================================================================================================
+#endif
+////////////////////////////////////////////////////////////////////////////////////////////////////////////////////////////////
+////////////////////////////////////////////////////////////////////////////////////////////////////////////////////////////////
+////////////////////////////////////////////////////////////////////////////////////////////////////////////////////////////////
+////////////////////////////////////////////////////////////////////////////////////////////////////////////////////////////////
+////////////////////////////////////////////////////////////////////////////////////////////////////////////////////////////////
+////////////////////////////////////////////////////////////////////////////////////////////////////////////////////////////////
+//_____________________________________________________________/\_______________________________________________________________
+//==============================================================================================================================
+//
+//
+//                                                            HLSL
+//
+//
+//==============================================================================================================================
+#if defined(A_HLSL) && defined(A_GPU)
+ #ifdef A_HLSL_6_2
+  #define AP1 bool
+  #define AP2 bool2
+  #define AP3 bool3
+  #define AP4 bool4
+//------------------------------------------------------------------------------------------------------------------------------
+  #define AF1 float32_t
+  #define AF2 float32_t2
+  #define AF3 float32_t3
+  #define AF4 float32_t4
+//------------------------------------------------------------------------------------------------------------------------------
+  #define AU1 uint32_t
+  #define AU2 uint32_t2
+  #define AU3 uint32_t3
+  #define AU4 uint32_t4
+//------------------------------------------------------------------------------------------------------------------------------
+  #define ASU1 int32_t
+  #define ASU2 int32_t2
+  #define ASU3 int32_t3
+  #define ASU4 int32_t4
+ #else
+  #define AP1 bool
+  #define AP2 bool2
+  #define AP3 bool3
+  #define AP4 bool4
+//------------------------------------------------------------------------------------------------------------------------------
+  #define AF1 float
+  #define AF2 float2
+  #define AF3 float3
+  #define AF4 float4
+//------------------------------------------------------------------------------------------------------------------------------
+  #define AU1 uint
+  #define AU2 uint2
+  #define AU3 uint3
+  #define AU4 uint4
+//------------------------------------------------------------------------------------------------------------------------------
+  #define ASU1 int
+  #define ASU2 int2
+  #define ASU3 int3
+  #define ASU4 int4
+ #endif
+//==============================================================================================================================
+ #define AF1_AU1(x) asfloat(AU1(x))
+ #define AF2_AU2(x) asfloat(AU2(x))
+ #define AF3_AU3(x) asfloat(AU3(x))
+ #define AF4_AU4(x) asfloat(AU4(x))
+//------------------------------------------------------------------------------------------------------------------------------
+ #define AU1_AF1(x) asuint(AF1(x))
+ #define AU2_AF2(x) asuint(AF2(x))
+ #define AU3_AF3(x) asuint(AF3(x))
+ #define AU4_AF4(x) asuint(AF4(x))
+//------------------------------------------------------------------------------------------------------------------------------
+ AU1 AU1_AH1_AF1_x(AF1 a){return f32tof16(a);}
+ #define AU1_AH1_AF1(a) AU1_AH1_AF1_x(AF1(a))
+//------------------------------------------------------------------------------------------------------------------------------
+ AU1 AU1_AH2_AF2_x(AF2 a){return f32tof16(a.x)|(f32tof16(a.y)<<16);}
+ #define AU1_AH2_AF2(a) AU1_AH2_AF2_x(AF2(a)) 
+ #define AU1_AB4Unorm_AF4(x) D3DCOLORtoUBYTE4(AF4(x))
+//------------------------------------------------------------------------------------------------------------------------------
+ AF2 AF2_AH2_AU1_x(AU1 x){return AF2(f16tof32(x&0xFFFF),f16tof32(x>>16));}
+ #define AF2_AH2_AU1(x) AF2_AH2_AU1_x(AU1(x))
+//==============================================================================================================================
+ AF1 AF1_x(AF1 a){return AF1(a);}
+ AF2 AF2_x(AF1 a){return AF2(a,a);}
+ AF3 AF3_x(AF1 a){return AF3(a,a,a);}
+ AF4 AF4_x(AF1 a){return AF4(a,a,a,a);}
+ #define AF1_(a) AF1_x(AF1(a))
+ #define AF2_(a) AF2_x(AF1(a))
+ #define AF3_(a) AF3_x(AF1(a))
+ #define AF4_(a) AF4_x(AF1(a))
+//------------------------------------------------------------------------------------------------------------------------------
+ AU1 AU1_x(AU1 a){return AU1(a);}
+ AU2 AU2_x(AU1 a){return AU2(a,a);}
+ AU3 AU3_x(AU1 a){return AU3(a,a,a);}
+ AU4 AU4_x(AU1 a){return AU4(a,a,a,a);}
+ #define AU1_(a) AU1_x(AU1(a))
+ #define AU2_(a) AU2_x(AU1(a))
+ #define AU3_(a) AU3_x(AU1(a))
+ #define AU4_(a) AU4_x(AU1(a))
+//==============================================================================================================================
+ AU1 AAbsSU1(AU1 a){return AU1(abs(ASU1(a)));}
+ AU2 AAbsSU2(AU2 a){return AU2(abs(ASU2(a)));}
+ AU3 AAbsSU3(AU3 a){return AU3(abs(ASU3(a)));}
+ AU4 AAbsSU4(AU4 a){return AU4(abs(ASU4(a)));}
+//------------------------------------------------------------------------------------------------------------------------------
+ AU1 ABfe(AU1 src,AU1 off,AU1 bits){AU1 mask=(1u<<bits)-1;return (src>>off)&mask;}
+ AU1 ABfi(AU1 src,AU1 ins,AU1 mask){return (ins&mask)|(src&(~mask));}
+ AU1 ABfiM(AU1 src,AU1 ins,AU1 bits){AU1 mask=(1u<<bits)-1;return (ins&mask)|(src&(~mask));}
+//------------------------------------------------------------------------------------------------------------------------------
+ AF1 AClampF1(AF1 x,AF1 n,AF1 m){return max(n,min(x,m));}
+ AF2 AClampF2(AF2 x,AF2 n,AF2 m){return max(n,min(x,m));}
+ AF3 AClampF3(AF3 x,AF3 n,AF3 m){return max(n,min(x,m));}
+ AF4 AClampF4(AF4 x,AF4 n,AF4 m){return max(n,min(x,m));}
+//------------------------------------------------------------------------------------------------------------------------------
+ AF1 AFractF1(AF1 x){return x-floor(x);}
+ AF2 AFractF2(AF2 x){return x-floor(x);}
+ AF3 AFractF3(AF3 x){return x-floor(x);}
+ AF4 AFractF4(AF4 x){return x-floor(x);}
+//------------------------------------------------------------------------------------------------------------------------------
+ AF1 ALerpF1(AF1 x,AF1 y,AF1 a){return lerp(x,y,a);}
+ AF2 ALerpF2(AF2 x,AF2 y,AF2 a){return lerp(x,y,a);}
+ AF3 ALerpF3(AF3 x,AF3 y,AF3 a){return lerp(x,y,a);}
+ AF4 ALerpF4(AF4 x,AF4 y,AF4 a){return lerp(x,y,a);}
+//------------------------------------------------------------------------------------------------------------------------------
+ AF1 AMax3F1(AF1 x,AF1 y,AF1 z){return max(x,max(y,z));}
+ AF2 AMax3F2(AF2 x,AF2 y,AF2 z){return max(x,max(y,z));}
+ AF3 AMax3F3(AF3 x,AF3 y,AF3 z){return max(x,max(y,z));}
+ AF4 AMax3F4(AF4 x,AF4 y,AF4 z){return max(x,max(y,z));}
+//------------------------------------------------------------------------------------------------------------------------------
+ AU1 AMax3SU1(AU1 x,AU1 y,AU1 z){return AU1(max(ASU1(x),max(ASU1(y),ASU1(z))));}
+ AU2 AMax3SU2(AU2 x,AU2 y,AU2 z){return AU2(max(ASU2(x),max(ASU2(y),ASU2(z))));}
+ AU3 AMax3SU3(AU3 x,AU3 y,AU3 z){return AU3(max(ASU3(x),max(ASU3(y),ASU3(z))));}
+ AU4 AMax3SU4(AU4 x,AU4 y,AU4 z){return AU4(max(ASU4(x),max(ASU4(y),ASU4(z))));}
+//------------------------------------------------------------------------------------------------------------------------------
+ AU1 AMax3U1(AU1 x,AU1 y,AU1 z){return max(x,max(y,z));}
+ AU2 AMax3U2(AU2 x,AU2 y,AU2 z){return max(x,max(y,z));}
+ AU3 AMax3U3(AU3 x,AU3 y,AU3 z){return max(x,max(y,z));}
+ AU4 AMax3U4(AU4 x,AU4 y,AU4 z){return max(x,max(y,z));}
+//------------------------------------------------------------------------------------------------------------------------------
+ AU1 AMaxSU1(AU1 a,AU1 b){return AU1(max(ASU1(a),ASU1(b)));}
+ AU2 AMaxSU2(AU2 a,AU2 b){return AU2(max(ASU2(a),ASU2(b)));}
+ AU3 AMaxSU3(AU3 a,AU3 b){return AU3(max(ASU3(a),ASU3(b)));}
+ AU4 AMaxSU4(AU4 a,AU4 b){return AU4(max(ASU4(a),ASU4(b)));}
+//------------------------------------------------------------------------------------------------------------------------------
+ AF1 AMed3F1(AF1 x,AF1 y,AF1 z){return max(min(x,y),min(max(x,y),z));}
+ AF2 AMed3F2(AF2 x,AF2 y,AF2 z){return max(min(x,y),min(max(x,y),z));}
+ AF3 AMed3F3(AF3 x,AF3 y,AF3 z){return max(min(x,y),min(max(x,y),z));}
+ AF4 AMed3F4(AF4 x,AF4 y,AF4 z){return max(min(x,y),min(max(x,y),z));}
+//------------------------------------------------------------------------------------------------------------------------------
+ AF1 AMin3F1(AF1 x,AF1 y,AF1 z){return min(x,min(y,z));}
+ AF2 AMin3F2(AF2 x,AF2 y,AF2 z){return min(x,min(y,z));}
+ AF3 AMin3F3(AF3 x,AF3 y,AF3 z){return min(x,min(y,z));}
+ AF4 AMin3F4(AF4 x,AF4 y,AF4 z){return min(x,min(y,z));}
+//------------------------------------------------------------------------------------------------------------------------------
+ AU1 AMin3SU1(AU1 x,AU1 y,AU1 z){return AU1(min(ASU1(x),min(ASU1(y),ASU1(z))));}
+ AU2 AMin3SU2(AU2 x,AU2 y,AU2 z){return AU2(min(ASU2(x),min(ASU2(y),ASU2(z))));}
+ AU3 AMin3SU3(AU3 x,AU3 y,AU3 z){return AU3(min(ASU3(x),min(ASU3(y),ASU3(z))));}
+ AU4 AMin3SU4(AU4 x,AU4 y,AU4 z){return AU4(min(ASU4(x),min(ASU4(y),ASU4(z))));}
+//------------------------------------------------------------------------------------------------------------------------------
+ AU1 AMin3U1(AU1 x,AU1 y,AU1 z){return min(x,min(y,z));}
+ AU2 AMin3U2(AU2 x,AU2 y,AU2 z){return min(x,min(y,z));}
+ AU3 AMin3U3(AU3 x,AU3 y,AU3 z){return min(x,min(y,z));}
+ AU4 AMin3U4(AU4 x,AU4 y,AU4 z){return min(x,min(y,z));}
+//------------------------------------------------------------------------------------------------------------------------------
+ AU1 AMinSU1(AU1 a,AU1 b){return AU1(min(ASU1(a),ASU1(b)));}
+ AU2 AMinSU2(AU2 a,AU2 b){return AU2(min(ASU2(a),ASU2(b)));}
+ AU3 AMinSU3(AU3 a,AU3 b){return AU3(min(ASU3(a),ASU3(b)));}
+ AU4 AMinSU4(AU4 a,AU4 b){return AU4(min(ASU4(a),ASU4(b)));}
+//------------------------------------------------------------------------------------------------------------------------------
+ AF1 ANCosF1(AF1 x){return cos(x*AF1_(A_2PI));}
+ AF2 ANCosF2(AF2 x){return cos(x*AF2_(A_2PI));}
+ AF3 ANCosF3(AF3 x){return cos(x*AF3_(A_2PI));}
+ AF4 ANCosF4(AF4 x){return cos(x*AF4_(A_2PI));}
+//------------------------------------------------------------------------------------------------------------------------------
+ AF1 ANSinF1(AF1 x){return sin(x*AF1_(A_2PI));}
+ AF2 ANSinF2(AF2 x){return sin(x*AF2_(A_2PI));}
+ AF3 ANSinF3(AF3 x){return sin(x*AF3_(A_2PI));}
+ AF4 ANSinF4(AF4 x){return sin(x*AF4_(A_2PI));}
+//------------------------------------------------------------------------------------------------------------------------------
+ AF1 ARcpF1(AF1 x){return rcp(x);}
+ AF2 ARcpF2(AF2 x){return rcp(x);}
+ AF3 ARcpF3(AF3 x){return rcp(x);}
+ AF4 ARcpF4(AF4 x){return rcp(x);}
+//------------------------------------------------------------------------------------------------------------------------------
+ AF1 ARsqF1(AF1 x){return rsqrt(x);}
+ AF2 ARsqF2(AF2 x){return rsqrt(x);}
+ AF3 ARsqF3(AF3 x){return rsqrt(x);}
+ AF4 ARsqF4(AF4 x){return rsqrt(x);}
+//------------------------------------------------------------------------------------------------------------------------------
+ AF1 ASatF1(AF1 x){return saturate(x);}
+ AF2 ASatF2(AF2 x){return saturate(x);}
+ AF3 ASatF3(AF3 x){return saturate(x);}
+ AF4 ASatF4(AF4 x){return saturate(x);}
+//------------------------------------------------------------------------------------------------------------------------------
+ AU1 AShrSU1(AU1 a,AU1 b){return AU1(ASU1(a)>>ASU1(b));}
+ AU2 AShrSU2(AU2 a,AU2 b){return AU2(ASU2(a)>>ASU2(b));}
+ AU3 AShrSU3(AU3 a,AU3 b){return AU3(ASU3(a)>>ASU3(b));}
+ AU4 AShrSU4(AU4 a,AU4 b){return AU4(ASU4(a)>>ASU4(b));}
+////////////////////////////////////////////////////////////////////////////////////////////////////////////////////////////////
+////////////////////////////////////////////////////////////////////////////////////////////////////////////////////////////////
+//_____________________________________________________________/\_______________________________________________________________
+//==============================================================================================================================
+//                                                          HLSL BYTE
+//==============================================================================================================================
+ #ifdef A_BYTE
+ #endif
+////////////////////////////////////////////////////////////////////////////////////////////////////////////////////////////////
+////////////////////////////////////////////////////////////////////////////////////////////////////////////////////////////////
+//_____________________________________________________________/\_______________________________________________________________
+//==============================================================================================================================
+//                                                          HLSL HALF
+//==============================================================================================================================
+ #ifdef A_HALF
+  #ifdef A_HLSL_6_2
+   #define AH1 float16_t
+   #define AH2 float16_t2
+   #define AH3 float16_t3
+   #define AH4 float16_t4
+//------------------------------------------------------------------------------------------------------------------------------
+   #define AW1 uint16_t
+   #define AW2 uint16_t2
+   #define AW3 uint16_t3
+   #define AW4 uint16_t4
+//------------------------------------------------------------------------------------------------------------------------------
+   #define ASW1 int16_t
+   #define ASW2 int16_t2
+   #define ASW3 int16_t3
+   #define ASW4 int16_t4
+  #else
+   #define AH1 min16float
+   #define AH2 min16float2
+   #define AH3 min16float3
+   #define AH4 min16float4
+//------------------------------------------------------------------------------------------------------------------------------
+   #define AW1 min16uint
+   #define AW2 min16uint2
+   #define AW3 min16uint3
+   #define AW4 min16uint4
+//------------------------------------------------------------------------------------------------------------------------------
+   #define ASW1 min16int
+   #define ASW2 min16int2
+   #define ASW3 min16int3
+   #define ASW4 min16int4
+  #endif
+//==============================================================================================================================
+  // Need to use manual unpack to get optimal execution (don't use packed types in buffers directly).
+  // Unpack requires this pattern: https://gpuopen.com/first-steps-implementing-fp16/
+  AH2 AH2_AU1_x(AU1 x){AF2 t=f16tof32(AU2(x&0xFFFF,x>>16));return AH2(t);}
+  AH4 AH4_AU2_x(AU2 x){return AH4(AH2_AU1_x(x.x),AH2_AU1_x(x.y));}
+  AW2 AW2_AU1_x(AU1 x){AU2 t=AU2(x&0xFFFF,x>>16);return AW2(t);}
+  AW4 AW4_AU2_x(AU2 x){return AW4(AW2_AU1_x(x.x),AW2_AU1_x(x.y));}
+  #define AH2_AU1(x) AH2_AU1_x(AU1(x))
+  #define AH4_AU2(x) AH4_AU2_x(AU2(x))
+  #define AW2_AU1(x) AW2_AU1_x(AU1(x))
+  #define AW4_AU2(x) AW4_AU2_x(AU2(x))
+//------------------------------------------------------------------------------------------------------------------------------
+  AU1 AU1_AH2_x(AH2 x){return f32tof16(x.x)+(f32tof16(x.y)<<16);}
+  AU2 AU2_AH4_x(AH4 x){return AU2(AU1_AH2_x(x.xy),AU1_AH2_x(x.zw));}
+  AU1 AU1_AW2_x(AW2 x){return AU1(x.x)+(AU1(x.y)<<16);}
+  AU2 AU2_AW4_x(AW4 x){return AU2(AU1_AW2_x(x.xy),AU1_AW2_x(x.zw));}
+  #define AU1_AH2(x) AU1_AH2_x(AH2(x))
+  #define AU2_AH4(x) AU2_AH4_x(AH4(x))
+  #define AU1_AW2(x) AU1_AW2_x(AW2(x))
+  #define AU2_AW4(x) AU2_AW4_x(AW4(x))
+//==============================================================================================================================
+  #if defined(A_HLSL_6_2) && !defined(A_NO_16_BIT_CAST)
+   #define AW1_AH1(x) asuint16(x)
+   #define AW2_AH2(x) asuint16(x)
+   #define AW3_AH3(x) asuint16(x)
+   #define AW4_AH4(x) asuint16(x)
+  #else
+   #define AW1_AH1(a) AW1(f32tof16(AF1(a)))
+   #define AW2_AH2(a) AW2(AW1_AH1((a).x),AW1_AH1((a).y))
+   #define AW3_AH3(a) AW3(AW1_AH1((a).x),AW1_AH1((a).y),AW1_AH1((a).z))
+   #define AW4_AH4(a) AW4(AW1_AH1((a).x),AW1_AH1((a).y),AW1_AH1((a).z),AW1_AH1((a).w))
+  #endif
+//------------------------------------------------------------------------------------------------------------------------------
+  #if defined(A_HLSL_6_2) && !defined(A_NO_16_BIT_CAST)
+   #define AH1_AW1(x) asfloat16(x)
+   #define AH2_AW2(x) asfloat16(x)
+   #define AH3_AW3(x) asfloat16(x)
+   #define AH4_AW4(x) asfloat16(x)
+  #else
+   #define AH1_AW1(a) AH1(f16tof32(AU1(a)))
+   #define AH2_AW2(a) AH2(AH1_AW1((a).x),AH1_AW1((a).y))
+   #define AH3_AW3(a) AH3(AH1_AW1((a).x),AH1_AW1((a).y),AH1_AW1((a).z))
+   #define AH4_AW4(a) AH4(AH1_AW1((a).x),AH1_AW1((a).y),AH1_AW1((a).z),AH1_AW1((a).w))
+  #endif
+//==============================================================================================================================
+  AH1 AH1_x(AH1 a){return AH1(a);}
+  AH2 AH2_x(AH1 a){return AH2(a,a);}
+  AH3 AH3_x(AH1 a){return AH3(a,a,a);}
+  AH4 AH4_x(AH1 a){return AH4(a,a,a,a);}
+  #define AH1_(a) AH1_x(AH1(a))
+  #define AH2_(a) AH2_x(AH1(a))
+  #define AH3_(a) AH3_x(AH1(a))
+  #define AH4_(a) AH4_x(AH1(a))
+//------------------------------------------------------------------------------------------------------------------------------
+  AW1 AW1_x(AW1 a){return AW1(a);}
+  AW2 AW2_x(AW1 a){return AW2(a,a);}
+  AW3 AW3_x(AW1 a){return AW3(a,a,a);}
+  AW4 AW4_x(AW1 a){return AW4(a,a,a,a);}
+  #define AW1_(a) AW1_x(AW1(a))
+  #define AW2_(a) AW2_x(AW1(a))
+  #define AW3_(a) AW3_x(AW1(a))
+  #define AW4_(a) AW4_x(AW1(a))
+//==============================================================================================================================
+  AW1 AAbsSW1(AW1 a){return AW1(abs(ASW1(a)));}
+  AW2 AAbsSW2(AW2 a){return AW2(abs(ASW2(a)));}
+  AW3 AAbsSW3(AW3 a){return AW3(abs(ASW3(a)));}
+  AW4 AAbsSW4(AW4 a){return AW4(abs(ASW4(a)));}
+//------------------------------------------------------------------------------------------------------------------------------
+  AH1 AClampH1(AH1 x,AH1 n,AH1 m){return max(n,min(x,m));}
+  AH2 AClampH2(AH2 x,AH2 n,AH2 m){return max(n,min(x,m));}
+  AH3 AClampH3(AH3 x,AH3 n,AH3 m){return max(n,min(x,m));}
+  AH4 AClampH4(AH4 x,AH4 n,AH4 m){return max(n,min(x,m));}
+//------------------------------------------------------------------------------------------------------------------------------
+ // V_FRACT_F16 (note DX frac() is different).
+  AH1 AFractH1(AH1 x){return x-floor(x);}
+  AH2 AFractH2(AH2 x){return x-floor(x);}
+  AH3 AFractH3(AH3 x){return x-floor(x);}
+  AH4 AFractH4(AH4 x){return x-floor(x);}
+//------------------------------------------------------------------------------------------------------------------------------
+  AH1 ALerpH1(AH1 x,AH1 y,AH1 a){return lerp(x,y,a);}
+  AH2 ALerpH2(AH2 x,AH2 y,AH2 a){return lerp(x,y,a);}
+  AH3 ALerpH3(AH3 x,AH3 y,AH3 a){return lerp(x,y,a);}
+  AH4 ALerpH4(AH4 x,AH4 y,AH4 a){return lerp(x,y,a);}
+//------------------------------------------------------------------------------------------------------------------------------
+  AH1 AMax3H1(AH1 x,AH1 y,AH1 z){return max(x,max(y,z));}
+  AH2 AMax3H2(AH2 x,AH2 y,AH2 z){return max(x,max(y,z));}
+  AH3 AMax3H3(AH3 x,AH3 y,AH3 z){return max(x,max(y,z));}
+  AH4 AMax3H4(AH4 x,AH4 y,AH4 z){return max(x,max(y,z));}
+//------------------------------------------------------------------------------------------------------------------------------
+  AW1 AMaxSW1(AW1 a,AW1 b){return AW1(max(ASU1(a),ASU1(b)));}
+  AW2 AMaxSW2(AW2 a,AW2 b){return AW2(max(ASU2(a),ASU2(b)));}
+  AW3 AMaxSW3(AW3 a,AW3 b){return AW3(max(ASU3(a),ASU3(b)));}
+  AW4 AMaxSW4(AW4 a,AW4 b){return AW4(max(ASU4(a),ASU4(b)));}
+//------------------------------------------------------------------------------------------------------------------------------
+  AH1 AMin3H1(AH1 x,AH1 y,AH1 z){return min(x,min(y,z));}
+  AH2 AMin3H2(AH2 x,AH2 y,AH2 z){return min(x,min(y,z));}
+  AH3 AMin3H3(AH3 x,AH3 y,AH3 z){return min(x,min(y,z));}
+  AH4 AMin3H4(AH4 x,AH4 y,AH4 z){return min(x,min(y,z));}
+//------------------------------------------------------------------------------------------------------------------------------
+  AW1 AMinSW1(AW1 a,AW1 b){return AW1(min(ASU1(a),ASU1(b)));}
+  AW2 AMinSW2(AW2 a,AW2 b){return AW2(min(ASU2(a),ASU2(b)));}
+  AW3 AMinSW3(AW3 a,AW3 b){return AW3(min(ASU3(a),ASU3(b)));}
+  AW4 AMinSW4(AW4 a,AW4 b){return AW4(min(ASU4(a),ASU4(b)));}
+//------------------------------------------------------------------------------------------------------------------------------
+  AH1 ARcpH1(AH1 x){return rcp(x);}
+  AH2 ARcpH2(AH2 x){return rcp(x);}
+  AH3 ARcpH3(AH3 x){return rcp(x);}
+  AH4 ARcpH4(AH4 x){return rcp(x);}
+//------------------------------------------------------------------------------------------------------------------------------
+  AH1 ARsqH1(AH1 x){return rsqrt(x);}
+  AH2 ARsqH2(AH2 x){return rsqrt(x);}
+  AH3 ARsqH3(AH3 x){return rsqrt(x);}
+  AH4 ARsqH4(AH4 x){return rsqrt(x);}
+//------------------------------------------------------------------------------------------------------------------------------
+  AH1 ASatH1(AH1 x){return saturate(x);}
+  AH2 ASatH2(AH2 x){return saturate(x);}
+  AH3 ASatH3(AH3 x){return saturate(x);}
+  AH4 ASatH4(AH4 x){return saturate(x);}
+//------------------------------------------------------------------------------------------------------------------------------
+  AW1 AShrSW1(AW1 a,AW1 b){return AW1(ASW1(a)>>ASW1(b));}
+  AW2 AShrSW2(AW2 a,AW2 b){return AW2(ASW2(a)>>ASW2(b));}
+  AW3 AShrSW3(AW3 a,AW3 b){return AW3(ASW3(a)>>ASW3(b));}
+  AW4 AShrSW4(AW4 a,AW4 b){return AW4(ASW4(a)>>ASW4(b));}
+ #endif
+////////////////////////////////////////////////////////////////////////////////////////////////////////////////////////////////
+////////////////////////////////////////////////////////////////////////////////////////////////////////////////////////////////
+//_____________________________________________________________/\_______________________________________________________________
+//==============================================================================================================================
+//                                                         HLSL DOUBLE
+//==============================================================================================================================
+ #ifdef A_DUBL
+  #ifdef A_HLSL_6_2
+   #define AD1 float64_t
+   #define AD2 float64_t2
+   #define AD3 float64_t3
+   #define AD4 float64_t4
+  #else
+   #define AD1 double
+   #define AD2 double2
+   #define AD3 double3
+   #define AD4 double4
+  #endif
+//------------------------------------------------------------------------------------------------------------------------------
+  AD1 AD1_x(AD1 a){return AD1(a);}
+  AD2 AD2_x(AD1 a){return AD2(a,a);}
+  AD3 AD3_x(AD1 a){return AD3(a,a,a);}
+  AD4 AD4_x(AD1 a){return AD4(a,a,a,a);}
+  #define AD1_(a) AD1_x(AD1(a))
+  #define AD2_(a) AD2_x(AD1(a))
+  #define AD3_(a) AD3_x(AD1(a))
+  #define AD4_(a) AD4_x(AD1(a))
+//==============================================================================================================================
+  AD1 AFractD1(AD1 a){return a-floor(a);}
+  AD2 AFractD2(AD2 a){return a-floor(a);}
+  AD3 AFractD3(AD3 a){return a-floor(a);}
+  AD4 AFractD4(AD4 a){return a-floor(a);}
+//------------------------------------------------------------------------------------------------------------------------------
+  AD1 ALerpD1(AD1 x,AD1 y,AD1 a){return lerp(x,y,a);}
+  AD2 ALerpD2(AD2 x,AD2 y,AD2 a){return lerp(x,y,a);}
+  AD3 ALerpD3(AD3 x,AD3 y,AD3 a){return lerp(x,y,a);}
+  AD4 ALerpD4(AD4 x,AD4 y,AD4 a){return lerp(x,y,a);}
+//------------------------------------------------------------------------------------------------------------------------------
+  AD1 ARcpD1(AD1 x){return rcp(x);}
+  AD2 ARcpD2(AD2 x){return rcp(x);}
+  AD3 ARcpD3(AD3 x){return rcp(x);}
+  AD4 ARcpD4(AD4 x){return rcp(x);}
+//------------------------------------------------------------------------------------------------------------------------------
+  AD1 ARsqD1(AD1 x){return rsqrt(x);}
+  AD2 ARsqD2(AD2 x){return rsqrt(x);}
+  AD3 ARsqD3(AD3 x){return rsqrt(x);}
+  AD4 ARsqD4(AD4 x){return rsqrt(x);}
+//------------------------------------------------------------------------------------------------------------------------------
+  AD1 ASatD1(AD1 x){return saturate(x);}
+  AD2 ASatD2(AD2 x){return saturate(x);}
+  AD3 ASatD3(AD3 x){return saturate(x);}
+  AD4 ASatD4(AD4 x){return saturate(x);}
+ #endif
+//==============================================================================================================================
+//                                                         HLSL WAVE
+//==============================================================================================================================
+ #ifdef A_WAVE
+  // Where 'x' must be a compile time literal.
+  AF1 AWaveXorF1(AF1 v,AU1 x){return WaveReadLaneAt(v,WaveGetLaneIndex()^x);}
+  AF2 AWaveXorF2(AF2 v,AU1 x){return WaveReadLaneAt(v,WaveGetLaneIndex()^x);}
+  AF3 AWaveXorF3(AF3 v,AU1 x){return WaveReadLaneAt(v,WaveGetLaneIndex()^x);}
+  AF4 AWaveXorF4(AF4 v,AU1 x){return WaveReadLaneAt(v,WaveGetLaneIndex()^x);}
+  AU1 AWaveXorU1(AU1 v,AU1 x){return WaveReadLaneAt(v,WaveGetLaneIndex()^x);}
+  AU2 AWaveXorU1(AU2 v,AU1 x){return WaveReadLaneAt(v,WaveGetLaneIndex()^x);}
+  AU3 AWaveXorU1(AU3 v,AU1 x){return WaveReadLaneAt(v,WaveGetLaneIndex()^x);}
+  AU4 AWaveXorU1(AU4 v,AU1 x){return WaveReadLaneAt(v,WaveGetLaneIndex()^x);}
+//------------------------------------------------------------------------------------------------------------------------------
+  #ifdef A_HALF
+   AH2 AWaveXorH2(AH2 v,AU1 x){return AH2_AU1(WaveReadLaneAt(AU1_AH2(v),WaveGetLaneIndex()^x));}
+   AH4 AWaveXorH4(AH4 v,AU1 x){return AH4_AU2(WaveReadLaneAt(AU2_AH4(v),WaveGetLaneIndex()^x));}
+   AW2 AWaveXorW2(AW2 v,AU1 x){return AW2_AU1(WaveReadLaneAt(AU1_AW2(v),WaveGetLaneIndex()^x));}
+   AW4 AWaveXorW4(AW4 v,AU1 x){return AW4_AU1(WaveReadLaneAt(AU1_AW4(v),WaveGetLaneIndex()^x));}
+  #endif
+ #endif
+//==============================================================================================================================
+#endif
+////////////////////////////////////////////////////////////////////////////////////////////////////////////////////////////////
+////////////////////////////////////////////////////////////////////////////////////////////////////////////////////////////////
+////////////////////////////////////////////////////////////////////////////////////////////////////////////////////////////////
+////////////////////////////////////////////////////////////////////////////////////////////////////////////////////////////////
+////////////////////////////////////////////////////////////////////////////////////////////////////////////////////////////////
+////////////////////////////////////////////////////////////////////////////////////////////////////////////////////////////////
+//_____________________________________________________________/\_______________________________________________________________
+//==============================================================================================================================
+//
+//
+//                                                          GPU COMMON
+//
+//
+//==============================================================================================================================
+#ifdef A_GPU
+ // Negative and positive infinity.
+ #define A_INFP_F AF1_AU1(0x7f800000u)
+ #define A_INFN_F AF1_AU1(0xff800000u)
+//------------------------------------------------------------------------------------------------------------------------------
+ // Copy sign from 's' to positive 'd'.
+ AF1 ACpySgnF1(AF1 d,AF1 s){return AF1_AU1(AU1_AF1(d)|(AU1_AF1(s)&AU1_(0x80000000u)));}
+ AF2 ACpySgnF2(AF2 d,AF2 s){return AF2_AU2(AU2_AF2(d)|(AU2_AF2(s)&AU2_(0x80000000u)));}
+ AF3 ACpySgnF3(AF3 d,AF3 s){return AF3_AU3(AU3_AF3(d)|(AU3_AF3(s)&AU3_(0x80000000u)));}
+ AF4 ACpySgnF4(AF4 d,AF4 s){return AF4_AU4(AU4_AF4(d)|(AU4_AF4(s)&AU4_(0x80000000u)));}
+//------------------------------------------------------------------------------------------------------------------------------
+ // Single operation to return (useful to create a mask to use in lerp for branch free logic),
+ //  m=NaN := 0
+ //  m>=0  := 0
+ //  m<0   := 1
+ // Uses the following useful floating point logic,
+ //  saturate(+a*(-INF)==-INF) := 0
+ //  saturate( 0*(-INF)== NaN) := 0
+ //  saturate(-a*(-INF)==+INF) := 1
+ AF1 ASignedF1(AF1 m){return ASatF1(m*AF1_(A_INFN_F));}
+ AF2 ASignedF2(AF2 m){return ASatF2(m*AF2_(A_INFN_F));}
+ AF3 ASignedF3(AF3 m){return ASatF3(m*AF3_(A_INFN_F));}
+ AF4 ASignedF4(AF4 m){return ASatF4(m*AF4_(A_INFN_F));}
+//------------------------------------------------------------------------------------------------------------------------------
+ AF1 AGtZeroF1(AF1 m){return ASatF1(m*AF1_(A_INFP_F));}
+ AF2 AGtZeroF2(AF2 m){return ASatF2(m*AF2_(A_INFP_F));}
+ AF3 AGtZeroF3(AF3 m){return ASatF3(m*AF3_(A_INFP_F));}
+ AF4 AGtZeroF4(AF4 m){return ASatF4(m*AF4_(A_INFP_F));}
+//==============================================================================================================================
+ #ifdef A_HALF
+  #ifdef A_HLSL_6_2
+   #define A_INFP_H AH1_AW1((uint16_t)0x7c00u)
+   #define A_INFN_H AH1_AW1((uint16_t)0xfc00u)
+  #else
+   #define A_INFP_H AH1_AW1(0x7c00u)
+   #define A_INFN_H AH1_AW1(0xfc00u)
+  #endif
+
+//------------------------------------------------------------------------------------------------------------------------------
+  AH1 ACpySgnH1(AH1 d,AH1 s){return AH1_AW1(AW1_AH1(d)|(AW1_AH1(s)&AW1_(0x8000u)));}
+  AH2 ACpySgnH2(AH2 d,AH2 s){return AH2_AW2(AW2_AH2(d)|(AW2_AH2(s)&AW2_(0x8000u)));}
+  AH3 ACpySgnH3(AH3 d,AH3 s){return AH3_AW3(AW3_AH3(d)|(AW3_AH3(s)&AW3_(0x8000u)));}
+  AH4 ACpySgnH4(AH4 d,AH4 s){return AH4_AW4(AW4_AH4(d)|(AW4_AH4(s)&AW4_(0x8000u)));}
+//------------------------------------------------------------------------------------------------------------------------------
+  AH1 ASignedH1(AH1 m){return ASatH1(m*AH1_(A_INFN_H));}
+  AH2 ASignedH2(AH2 m){return ASatH2(m*AH2_(A_INFN_H));}
+  AH3 ASignedH3(AH3 m){return ASatH3(m*AH3_(A_INFN_H));}
+  AH4 ASignedH4(AH4 m){return ASatH4(m*AH4_(A_INFN_H));}
+//------------------------------------------------------------------------------------------------------------------------------
+  AH1 AGtZeroH1(AH1 m){return ASatH1(m*AH1_(A_INFP_H));}
+  AH2 AGtZeroH2(AH2 m){return ASatH2(m*AH2_(A_INFP_H));}
+  AH3 AGtZeroH3(AH3 m){return ASatH3(m*AH3_(A_INFP_H));}
+  AH4 AGtZeroH4(AH4 m){return ASatH4(m*AH4_(A_INFP_H));}
+ #endif
+////////////////////////////////////////////////////////////////////////////////////////////////////////////////////////////////
+////////////////////////////////////////////////////////////////////////////////////////////////////////////////////////////////
+//_____________________________________________________________/\_______________________________________________________________
+//==============================================================================================================================
+//                                                [FIS] FLOAT INTEGER SORTABLE
+//------------------------------------------------------------------------------------------------------------------------------
+// Float to integer sortable.
+//  - If sign bit=0, flip the sign bit (positives).
+//  - If sign bit=1, flip all bits     (negatives).
+// Integer sortable to float.
+//  - If sign bit=1, flip the sign bit (positives).
+//  - If sign bit=0, flip all bits     (negatives).
+// Has nice side effects.
+//  - Larger integers are more positive values.
+//  - Float zero is mapped to center of integers (so clear to integer zero is a nice default for atomic max usage).
+// Burns 3 ops for conversion {shift,or,xor}.
+//==============================================================================================================================
+ AU1 AFisToU1(AU1 x){return x^(( AShrSU1(x,AU1_(31)))|AU1_(0x80000000));}
+ AU1 AFisFromU1(AU1 x){return x^((~AShrSU1(x,AU1_(31)))|AU1_(0x80000000));}
+//------------------------------------------------------------------------------------------------------------------------------
+ // Just adjust high 16-bit value (useful when upper part of 32-bit word is a 16-bit float value).
+ AU1 AFisToHiU1(AU1 x){return x^(( AShrSU1(x,AU1_(15)))|AU1_(0x80000000));}
+ AU1 AFisFromHiU1(AU1 x){return x^((~AShrSU1(x,AU1_(15)))|AU1_(0x80000000));}
+//------------------------------------------------------------------------------------------------------------------------------
+ #ifdef A_HALF
+  AW1 AFisToW1(AW1 x){return x^(( AShrSW1(x,AW1_(15)))|AW1_(0x8000));}
+  AW1 AFisFromW1(AW1 x){return x^((~AShrSW1(x,AW1_(15)))|AW1_(0x8000));}
+//------------------------------------------------------------------------------------------------------------------------------
+  AW2 AFisToW2(AW2 x){return x^(( AShrSW2(x,AW2_(15)))|AW2_(0x8000));}
+  AW2 AFisFromW2(AW2 x){return x^((~AShrSW2(x,AW2_(15)))|AW2_(0x8000));}
+ #endif
+////////////////////////////////////////////////////////////////////////////////////////////////////////////////////////////////
+////////////////////////////////////////////////////////////////////////////////////////////////////////////////////////////////
+//_____________________________________________________________/\_______________________________________________________________
+//==============================================================================================================================
+//                                                      [PERM] V_PERM_B32
+//------------------------------------------------------------------------------------------------------------------------------
+// Support for V_PERM_B32 started in the 3rd generation of GCN.
+//------------------------------------------------------------------------------------------------------------------------------
+// yyyyxxxx - The 'i' input.
+// 76543210
+// ========
+// HGFEDCBA - Naming on permutation.
+//------------------------------------------------------------------------------------------------------------------------------
+// TODO
+// ====
+//  - Make sure compiler optimizes this.
+//==============================================================================================================================
+ #ifdef A_HALF
+  AU1 APerm0E0A(AU2 i){return((i.x    )&0xffu)|((i.y<<16)&0xff0000u);}
+  AU1 APerm0F0B(AU2 i){return((i.x>> 8)&0xffu)|((i.y<< 8)&0xff0000u);}
+  AU1 APerm0G0C(AU2 i){return((i.x>>16)&0xffu)|((i.y    )&0xff0000u);}
+  AU1 APerm0H0D(AU2 i){return((i.x>>24)&0xffu)|((i.y>> 8)&0xff0000u);}
+//------------------------------------------------------------------------------------------------------------------------------
+  AU1 APermHGFA(AU2 i){return((i.x    )&0x000000ffu)|(i.y&0xffffff00u);}
+  AU1 APermHGFC(AU2 i){return((i.x>>16)&0x000000ffu)|(i.y&0xffffff00u);}
+  AU1 APermHGAE(AU2 i){return((i.x<< 8)&0x0000ff00u)|(i.y&0xffff00ffu);}
+  AU1 APermHGCE(AU2 i){return((i.x>> 8)&0x0000ff00u)|(i.y&0xffff00ffu);}
+  AU1 APermHAFE(AU2 i){return((i.x<<16)&0x00ff0000u)|(i.y&0xff00ffffu);}
+  AU1 APermHCFE(AU2 i){return((i.x    )&0x00ff0000u)|(i.y&0xff00ffffu);}
+  AU1 APermAGFE(AU2 i){return((i.x<<24)&0xff000000u)|(i.y&0x00ffffffu);}
+  AU1 APermCGFE(AU2 i){return((i.x<< 8)&0xff000000u)|(i.y&0x00ffffffu);}
+//------------------------------------------------------------------------------------------------------------------------------
+  AU1 APermGCEA(AU2 i){return((i.x)&0x00ff00ffu)|((i.y<<8)&0xff00ff00u);}
+  AU1 APermGECA(AU2 i){return(((i.x)&0xffu)|((i.x>>8)&0xff00u)|((i.y<<16)&0xff0000u)|((i.y<<8)&0xff000000u));}
+ #endif
+////////////////////////////////////////////////////////////////////////////////////////////////////////////////////////////////
+////////////////////////////////////////////////////////////////////////////////////////////////////////////////////////////////
+//_____________________________________________________________/\_______________________________________________________________
+//==============================================================================================================================
+//                                               [BUC] BYTE UNSIGNED CONVERSION
+//------------------------------------------------------------------------------------------------------------------------------
+// Designed to use the optimal conversion, enables the scaling to possibly be factored into other computation.
+// Works on a range of {0 to A_BUC_<32,16>}, for <32-bit, and 16-bit> respectively.
+//------------------------------------------------------------------------------------------------------------------------------
+// OPCODE NOTES
+// ============
+// GCN does not do UNORM or SNORM for bytes in opcodes.
+//  - V_CVT_F32_UBYTE{0,1,2,3} - Unsigned byte to float.
+//  - V_CVT_PKACC_U8_F32 - Float to unsigned byte (does bit-field insert into 32-bit integer).
+// V_PERM_B32 does byte packing with ability to zero fill bytes as well.
+//  - Can pull out byte values from two sources, and zero fill upper 8-bits of packed hi and lo. 
+//------------------------------------------------------------------------------------------------------------------------------
+// BYTE : FLOAT - ABuc{0,1,2,3}{To,From}U1() - Designed for V_CVT_F32_UBYTE* and V_CVT_PKACCUM_U8_F32 ops.
+// ====   =====
+//    0 : 0
+//    1 : 1
+//     ...
+//  255 : 255
+//      : 256 (just outside the encoding range)
+//------------------------------------------------------------------------------------------------------------------------------
+// BYTE : FLOAT - ABuc{0,1,2,3}{To,From}U2() - Designed for 16-bit denormal tricks and V_PERM_B32.
+// ====   =====
+//    0 : 0
+//    1 : 1/512
+//    2 : 1/256
+//     ...
+//   64 : 1/8
+//  128 : 1/4
+//  255 : 255/512
+//      : 1/2 (just outside the encoding range)
+//------------------------------------------------------------------------------------------------------------------------------
+// OPTIMAL IMPLEMENTATIONS ON AMD ARCHITECTURES
+// ============================================
+// r=ABuc0FromU1(i)
+//   V_CVT_F32_UBYTE0 r,i
+// --------------------------------------------
+// r=ABuc0ToU1(d,i)
+//   V_CVT_PKACCUM_U8_F32 r,i,0,d
+// --------------------------------------------
+// d=ABuc0FromU2(i)
+//   Where 'k0' is an SGPR with 0x0E0A
+//   Where 'k1' is an SGPR with {32768.0} packed into the lower 16-bits
+//   V_PERM_B32 d,i.x,i.y,k0
+//   V_PK_FMA_F16 d,d,k1.x,0
+// --------------------------------------------
+// r=ABuc0ToU2(d,i)
+//   Where 'k0' is an SGPR with {1.0/32768.0} packed into the lower 16-bits
+//   Where 'k1' is an SGPR with 0x????
+//   Where 'k2' is an SGPR with 0x????
+//   V_PK_FMA_F16 i,i,k0.x,0
+//   V_PERM_B32 r.x,i,i,k1
+//   V_PERM_B32 r.y,i,i,k2
+//==============================================================================================================================
+ // Peak range for 32-bit and 16-bit operations.
+ #define A_BUC_32 (255.0)
+ #define A_BUC_16 (255.0/512.0)
+//==============================================================================================================================
+ #if 1
+  // Designed to be one V_CVT_PKACCUM_U8_F32.
+  // The extra min is required to pattern match to V_CVT_PKACCUM_U8_F32.
+  AU1 ABuc0ToU1(AU1 d,AF1 i){return (d&0xffffff00u)|((min(AU1(i),255u)    )&(0x000000ffu));}
+  AU1 ABuc1ToU1(AU1 d,AF1 i){return (d&0xffff00ffu)|((min(AU1(i),255u)<< 8)&(0x0000ff00u));}
+  AU1 ABuc2ToU1(AU1 d,AF1 i){return (d&0xff00ffffu)|((min(AU1(i),255u)<<16)&(0x00ff0000u));}
+  AU1 ABuc3ToU1(AU1 d,AF1 i){return (d&0x00ffffffu)|((min(AU1(i),255u)<<24)&(0xff000000u));}
+//------------------------------------------------------------------------------------------------------------------------------
+  // Designed to be one V_CVT_F32_UBYTE*.
+  AF1 ABuc0FromU1(AU1 i){return AF1((i    )&255u);}
+  AF1 ABuc1FromU1(AU1 i){return AF1((i>> 8)&255u);}
+  AF1 ABuc2FromU1(AU1 i){return AF1((i>>16)&255u);}
+  AF1 ABuc3FromU1(AU1 i){return AF1((i>>24)&255u);}
+ #endif
+//==============================================================================================================================
+ #ifdef A_HALF
+  // Takes {x0,x1} and {y0,y1} and builds {{x0,y0},{x1,y1}}.
+  AW2 ABuc01ToW2(AH2 x,AH2 y){x*=AH2_(1.0/32768.0);y*=AH2_(1.0/32768.0);
+   return AW2_AU1(APermGCEA(AU2(AU1_AW2(AW2_AH2(x)),AU1_AW2(AW2_AH2(y)))));}
+//------------------------------------------------------------------------------------------------------------------------------
+  // Designed for 3 ops to do SOA to AOS and conversion.
+  AU2 ABuc0ToU2(AU2 d,AH2 i){AU1 b=AU1_AW2(AW2_AH2(i*AH2_(1.0/32768.0)));
+   return AU2(APermHGFA(AU2(d.x,b)),APermHGFC(AU2(d.y,b)));}
+  AU2 ABuc1ToU2(AU2 d,AH2 i){AU1 b=AU1_AW2(AW2_AH2(i*AH2_(1.0/32768.0)));
+   return AU2(APermHGAE(AU2(d.x,b)),APermHGCE(AU2(d.y,b)));}
+  AU2 ABuc2ToU2(AU2 d,AH2 i){AU1 b=AU1_AW2(AW2_AH2(i*AH2_(1.0/32768.0)));
+   return AU2(APermHAFE(AU2(d.x,b)),APermHCFE(AU2(d.y,b)));}
+  AU2 ABuc3ToU2(AU2 d,AH2 i){AU1 b=AU1_AW2(AW2_AH2(i*AH2_(1.0/32768.0)));
+   return AU2(APermAGFE(AU2(d.x,b)),APermCGFE(AU2(d.y,b)));}
+//------------------------------------------------------------------------------------------------------------------------------
+  // Designed for 2 ops to do both AOS to SOA, and conversion.
+  AH2 ABuc0FromU2(AU2 i){return AH2_AW2(AW2_AU1(APerm0E0A(i)))*AH2_(32768.0);}
+  AH2 ABuc1FromU2(AU2 i){return AH2_AW2(AW2_AU1(APerm0F0B(i)))*AH2_(32768.0);}
+  AH2 ABuc2FromU2(AU2 i){return AH2_AW2(AW2_AU1(APerm0G0C(i)))*AH2_(32768.0);}
+  AH2 ABuc3FromU2(AU2 i){return AH2_AW2(AW2_AU1(APerm0H0D(i)))*AH2_(32768.0);}
+ #endif
+////////////////////////////////////////////////////////////////////////////////////////////////////////////////////////////////
+////////////////////////////////////////////////////////////////////////////////////////////////////////////////////////////////
+//_____________________________________________________________/\_______________________________________________________________
+//==============================================================================================================================
+//                                                 [BSC] BYTE SIGNED CONVERSION
+//------------------------------------------------------------------------------------------------------------------------------
+// Similar to [BUC].
+// Works on a range of {-/+ A_BSC_<32,16>}, for <32-bit, and 16-bit> respectively.
+//------------------------------------------------------------------------------------------------------------------------------
+// ENCODING (without zero-based encoding)
+// ========
+//   0 = unused (can be used to mean something else)
+//   1 = lowest value 
+// 128 = exact zero center (zero based encoding 
+// 255 = highest value
+//------------------------------------------------------------------------------------------------------------------------------
+// Zero-based [Zb] flips the MSB bit of the byte (making 128 "exact zero" actually zero).
+// This is useful if there is a desire for cleared values to decode as zero.
+//------------------------------------------------------------------------------------------------------------------------------
+// BYTE : FLOAT - ABsc{0,1,2,3}{To,From}U2() - Designed for 16-bit denormal tricks and V_PERM_B32.
+// ====   =====
+//    0 : -127/512 (unused)
+//    1 : -126/512
+//    2 : -125/512
+//     ...
+//  128 : 0 
+//     ... 
+//  255 : 127/512
+//      : 1/4 (just outside the encoding range)
+//==============================================================================================================================
+ // Peak range for 32-bit and 16-bit operations.
+ #define A_BSC_32 (127.0)
+ #define A_BSC_16 (127.0/512.0)
+//==============================================================================================================================
+ #if 1
+  AU1 ABsc0ToU1(AU1 d,AF1 i){return (d&0xffffff00u)|((min(AU1(i+128.0),255u)    )&(0x000000ffu));}
+  AU1 ABsc1ToU1(AU1 d,AF1 i){return (d&0xffff00ffu)|((min(AU1(i+128.0),255u)<< 8)&(0x0000ff00u));}
+  AU1 ABsc2ToU1(AU1 d,AF1 i){return (d&0xff00ffffu)|((min(AU1(i+128.0),255u)<<16)&(0x00ff0000u));}
+  AU1 ABsc3ToU1(AU1 d,AF1 i){return (d&0x00ffffffu)|((min(AU1(i+128.0),255u)<<24)&(0xff000000u));}
+//------------------------------------------------------------------------------------------------------------------------------
+  AU1 ABsc0ToZbU1(AU1 d,AF1 i){return ((d&0xffffff00u)|((min(AU1(trunc(i)+128.0),255u)    )&(0x000000ffu)))^0x00000080u;}
+  AU1 ABsc1ToZbU1(AU1 d,AF1 i){return ((d&0xffff00ffu)|((min(AU1(trunc(i)+128.0),255u)<< 8)&(0x0000ff00u)))^0x00008000u;}
+  AU1 ABsc2ToZbU1(AU1 d,AF1 i){return ((d&0xff00ffffu)|((min(AU1(trunc(i)+128.0),255u)<<16)&(0x00ff0000u)))^0x00800000u;}
+  AU1 ABsc3ToZbU1(AU1 d,AF1 i){return ((d&0x00ffffffu)|((min(AU1(trunc(i)+128.0),255u)<<24)&(0xff000000u)))^0x80000000u;}
+//------------------------------------------------------------------------------------------------------------------------------
+  AF1 ABsc0FromU1(AU1 i){return AF1((i    )&255u)-128.0;}
+  AF1 ABsc1FromU1(AU1 i){return AF1((i>> 8)&255u)-128.0;}
+  AF1 ABsc2FromU1(AU1 i){return AF1((i>>16)&255u)-128.0;}
+  AF1 ABsc3FromU1(AU1 i){return AF1((i>>24)&255u)-128.0;}
+//------------------------------------------------------------------------------------------------------------------------------
+  AF1 ABsc0FromZbU1(AU1 i){return AF1(((i    )&255u)^0x80u)-128.0;}
+  AF1 ABsc1FromZbU1(AU1 i){return AF1(((i>> 8)&255u)^0x80u)-128.0;}
+  AF1 ABsc2FromZbU1(AU1 i){return AF1(((i>>16)&255u)^0x80u)-128.0;}
+  AF1 ABsc3FromZbU1(AU1 i){return AF1(((i>>24)&255u)^0x80u)-128.0;}
+ #endif
+//==============================================================================================================================
+ #ifdef A_HALF
+  // Takes {x0,x1} and {y0,y1} and builds {{x0,y0},{x1,y1}}.
+  AW2 ABsc01ToW2(AH2 x,AH2 y){x=x*AH2_(1.0/32768.0)+AH2_(0.25/32768.0);y=y*AH2_(1.0/32768.0)+AH2_(0.25/32768.0);
+   return AW2_AU1(APermGCEA(AU2(AU1_AW2(AW2_AH2(x)),AU1_AW2(AW2_AH2(y)))));}
+//------------------------------------------------------------------------------------------------------------------------------
+  AU2 ABsc0ToU2(AU2 d,AH2 i){AU1 b=AU1_AW2(AW2_AH2(i*AH2_(1.0/32768.0)+AH2_(0.25/32768.0)));
+   return AU2(APermHGFA(AU2(d.x,b)),APermHGFC(AU2(d.y,b)));}
+  AU2 ABsc1ToU2(AU2 d,AH2 i){AU1 b=AU1_AW2(AW2_AH2(i*AH2_(1.0/32768.0)+AH2_(0.25/32768.0)));
+   return AU2(APermHGAE(AU2(d.x,b)),APermHGCE(AU2(d.y,b)));}
+  AU2 ABsc2ToU2(AU2 d,AH2 i){AU1 b=AU1_AW2(AW2_AH2(i*AH2_(1.0/32768.0)+AH2_(0.25/32768.0)));
+   return AU2(APermHAFE(AU2(d.x,b)),APermHCFE(AU2(d.y,b)));}
+  AU2 ABsc3ToU2(AU2 d,AH2 i){AU1 b=AU1_AW2(AW2_AH2(i*AH2_(1.0/32768.0)+AH2_(0.25/32768.0)));
+   return AU2(APermAGFE(AU2(d.x,b)),APermCGFE(AU2(d.y,b)));}
+//------------------------------------------------------------------------------------------------------------------------------
+  AU2 ABsc0ToZbU2(AU2 d,AH2 i){AU1 b=AU1_AW2(AW2_AH2(i*AH2_(1.0/32768.0)+AH2_(0.25/32768.0)))^0x00800080u;
+   return AU2(APermHGFA(AU2(d.x,b)),APermHGFC(AU2(d.y,b)));}
+  AU2 ABsc1ToZbU2(AU2 d,AH2 i){AU1 b=AU1_AW2(AW2_AH2(i*AH2_(1.0/32768.0)+AH2_(0.25/32768.0)))^0x00800080u;
+   return AU2(APermHGAE(AU2(d.x,b)),APermHGCE(AU2(d.y,b)));}
+  AU2 ABsc2ToZbU2(AU2 d,AH2 i){AU1 b=AU1_AW2(AW2_AH2(i*AH2_(1.0/32768.0)+AH2_(0.25/32768.0)))^0x00800080u;
+   return AU2(APermHAFE(AU2(d.x,b)),APermHCFE(AU2(d.y,b)));}
+  AU2 ABsc3ToZbU2(AU2 d,AH2 i){AU1 b=AU1_AW2(AW2_AH2(i*AH2_(1.0/32768.0)+AH2_(0.25/32768.0)))^0x00800080u;
+   return AU2(APermAGFE(AU2(d.x,b)),APermCGFE(AU2(d.y,b)));}
+//------------------------------------------------------------------------------------------------------------------------------
+  AH2 ABsc0FromU2(AU2 i){return AH2_AW2(AW2_AU1(APerm0E0A(i)))*AH2_(32768.0)-AH2_(0.25);}
+  AH2 ABsc1FromU2(AU2 i){return AH2_AW2(AW2_AU1(APerm0F0B(i)))*AH2_(32768.0)-AH2_(0.25);}
+  AH2 ABsc2FromU2(AU2 i){return AH2_AW2(AW2_AU1(APerm0G0C(i)))*AH2_(32768.0)-AH2_(0.25);}
+  AH2 ABsc3FromU2(AU2 i){return AH2_AW2(AW2_AU1(APerm0H0D(i)))*AH2_(32768.0)-AH2_(0.25);}
+//------------------------------------------------------------------------------------------------------------------------------
+  AH2 ABsc0FromZbU2(AU2 i){return AH2_AW2(AW2_AU1(APerm0E0A(i)^0x00800080u))*AH2_(32768.0)-AH2_(0.25);}
+  AH2 ABsc1FromZbU2(AU2 i){return AH2_AW2(AW2_AU1(APerm0F0B(i)^0x00800080u))*AH2_(32768.0)-AH2_(0.25);}
+  AH2 ABsc2FromZbU2(AU2 i){return AH2_AW2(AW2_AU1(APerm0G0C(i)^0x00800080u))*AH2_(32768.0)-AH2_(0.25);}
+  AH2 ABsc3FromZbU2(AU2 i){return AH2_AW2(AW2_AU1(APerm0H0D(i)^0x00800080u))*AH2_(32768.0)-AH2_(0.25);}
+ #endif
+////////////////////////////////////////////////////////////////////////////////////////////////////////////////////////////////
+////////////////////////////////////////////////////////////////////////////////////////////////////////////////////////////////
+//_____________________________________________________________/\_______________________________________________________________
+//==============================================================================================================================
+//                                                     HALF APPROXIMATIONS
+//------------------------------------------------------------------------------------------------------------------------------
+// These support only positive inputs.
+// Did not see value yet in specialization for range.
+// Using quick testing, ended up mostly getting the same "best" approximation for various ranges.
+// With hardware that can co-execute transcendentals, the value in approximations could be less than expected.
+// However from a latency perspective, if execution of a transcendental is 4 clk, with no packed support, -> 8 clk total.
+// And co-execution would require a compiler interleaving a lot of independent work for packed usage.
+//------------------------------------------------------------------------------------------------------------------------------
+// The one Newton Raphson iteration form of rsq() was skipped (requires 6 ops total).
+// Same with sqrt(), as this could be x*rsq() (7 ops).
+//==============================================================================================================================
+ #ifdef A_HALF
+  // Minimize squared error across full positive range, 2 ops.
+  // The 0x1de2 based approximation maps {0 to 1} input maps to < 1 output.
+  AH1 APrxLoSqrtH1(AH1 a){return AH1_AW1((AW1_AH1(a)>>AW1_(1))+AW1_(0x1de2));}
+  AH2 APrxLoSqrtH2(AH2 a){return AH2_AW2((AW2_AH2(a)>>AW2_(1))+AW2_(0x1de2));}
+  AH3 APrxLoSqrtH3(AH3 a){return AH3_AW3((AW3_AH3(a)>>AW3_(1))+AW3_(0x1de2));}
+  AH4 APrxLoSqrtH4(AH4 a){return AH4_AW4((AW4_AH4(a)>>AW4_(1))+AW4_(0x1de2));}
+//------------------------------------------------------------------------------------------------------------------------------
+  // Lower precision estimation, 1 op.
+  // Minimize squared error across {smallest normal to 16384.0}.
+  AH1 APrxLoRcpH1(AH1 a){return AH1_AW1(AW1_(0x7784)-AW1_AH1(a));}
+  AH2 APrxLoRcpH2(AH2 a){return AH2_AW2(AW2_(0x7784)-AW2_AH2(a));}
+  AH3 APrxLoRcpH3(AH3 a){return AH3_AW3(AW3_(0x7784)-AW3_AH3(a));}
+  AH4 APrxLoRcpH4(AH4 a){return AH4_AW4(AW4_(0x7784)-AW4_AH4(a));}
+//------------------------------------------------------------------------------------------------------------------------------
+  // Medium precision estimation, one Newton Raphson iteration, 3 ops.
+  AH1 APrxMedRcpH1(AH1 a){AH1 b=AH1_AW1(AW1_(0x778d)-AW1_AH1(a));return b*(-b*a+AH1_(2.0));}
+  AH2 APrxMedRcpH2(AH2 a){AH2 b=AH2_AW2(AW2_(0x778d)-AW2_AH2(a));return b*(-b*a+AH2_(2.0));}
+  AH3 APrxMedRcpH3(AH3 a){AH3 b=AH3_AW3(AW3_(0x778d)-AW3_AH3(a));return b*(-b*a+AH3_(2.0));}
+  AH4 APrxMedRcpH4(AH4 a){AH4 b=AH4_AW4(AW4_(0x778d)-AW4_AH4(a));return b*(-b*a+AH4_(2.0));}
+//------------------------------------------------------------------------------------------------------------------------------
+  // Minimize squared error across {smallest normal to 16384.0}, 2 ops.
+  AH1 APrxLoRsqH1(AH1 a){return AH1_AW1(AW1_(0x59a3)-(AW1_AH1(a)>>AW1_(1)));}
+  AH2 APrxLoRsqH2(AH2 a){return AH2_AW2(AW2_(0x59a3)-(AW2_AH2(a)>>AW2_(1)));}
+  AH3 APrxLoRsqH3(AH3 a){return AH3_AW3(AW3_(0x59a3)-(AW3_AH3(a)>>AW3_(1)));}
+  AH4 APrxLoRsqH4(AH4 a){return AH4_AW4(AW4_(0x59a3)-(AW4_AH4(a)>>AW4_(1)));}
+ #endif
+////////////////////////////////////////////////////////////////////////////////////////////////////////////////////////////////
+////////////////////////////////////////////////////////////////////////////////////////////////////////////////////////////////
+//_____________________________________________________________/\_______________________________________________________________
+//==============================================================================================================================
+//                                                    FLOAT APPROXIMATIONS
+//------------------------------------------------------------------------------------------------------------------------------
+// Michal Drobot has an excellent presentation on these: "Low Level Optimizations For GCN",
+//  - Idea dates back to SGI, then to Quake 3, etc.
+//  - https://michaldrobot.files.wordpress.com/2014/05/gcn_alu_opt_digitaldragons2014.pdf
+//     - sqrt(x)=rsqrt(x)*x
+//     - rcp(x)=rsqrt(x)*rsqrt(x) for positive x
+//  - [REDACTED_PROFILE_URL]/ShaderFastLibs/blob/master/ShaderFastMathLib.h
+//------------------------------------------------------------------------------------------------------------------------------
+// These below are from perhaps less complete searching for optimal.
+// Used FP16 normal range for testing with +4096 32-bit step size for sampling error.
+// So these match up well with the half approximations.
+//==============================================================================================================================
+ AF1 APrxLoSqrtF1(AF1 a){return AF1_AU1((AU1_AF1(a)>>AU1_(1))+AU1_(0x1fbc4639));}
+ AF1 APrxLoRcpF1(AF1 a){return AF1_AU1(AU1_(0x7ef07ebb)-AU1_AF1(a));}
+ AF1 APrxMedRcpF1(AF1 a){AF1 b=AF1_AU1(AU1_(0x7ef19fff)-AU1_AF1(a));return b*(-b*a+AF1_(2.0));}
+ AF1 APrxLoRsqF1(AF1 a){return AF1_AU1(AU1_(0x5f347d74)-(AU1_AF1(a)>>AU1_(1)));}
+//------------------------------------------------------------------------------------------------------------------------------
+ AF2 APrxLoSqrtF2(AF2 a){return AF2_AU2((AU2_AF2(a)>>AU2_(1))+AU2_(0x1fbc4639));}
+ AF2 APrxLoRcpF2(AF2 a){return AF2_AU2(AU2_(0x7ef07ebb)-AU2_AF2(a));}
+ AF2 APrxMedRcpF2(AF2 a){AF2 b=AF2_AU2(AU2_(0x7ef19fff)-AU2_AF2(a));return b*(-b*a+AF2_(2.0));}
+ AF2 APrxLoRsqF2(AF2 a){return AF2_AU2(AU2_(0x5f347d74)-(AU2_AF2(a)>>AU2_(1)));}
+//------------------------------------------------------------------------------------------------------------------------------
+ AF3 APrxLoSqrtF3(AF3 a){return AF3_AU3((AU3_AF3(a)>>AU3_(1))+AU3_(0x1fbc4639));}
+ AF3 APrxLoRcpF3(AF3 a){return AF3_AU3(AU3_(0x7ef07ebb)-AU3_AF3(a));}
+ AF3 APrxMedRcpF3(AF3 a){AF3 b=AF3_AU3(AU3_(0x7ef19fff)-AU3_AF3(a));return b*(-b*a+AF3_(2.0));}
+ AF3 APrxLoRsqF3(AF3 a){return AF3_AU3(AU3_(0x5f347d74)-(AU3_AF3(a)>>AU3_(1)));}
+//------------------------------------------------------------------------------------------------------------------------------
+ AF4 APrxLoSqrtF4(AF4 a){return AF4_AU4((AU4_AF4(a)>>AU4_(1))+AU4_(0x1fbc4639));}
+ AF4 APrxLoRcpF4(AF4 a){return AF4_AU4(AU4_(0x7ef07ebb)-AU4_AF4(a));}
+ AF4 APrxMedRcpF4(AF4 a){AF4 b=AF4_AU4(AU4_(0x7ef19fff)-AU4_AF4(a));return b*(-b*a+AF4_(2.0));}
+ AF4 APrxLoRsqF4(AF4 a){return AF4_AU4(AU4_(0x5f347d74)-(AU4_AF4(a)>>AU4_(1)));}
+////////////////////////////////////////////////////////////////////////////////////////////////////////////////////////////////
+////////////////////////////////////////////////////////////////////////////////////////////////////////////////////////////////
+//_____________________________________________________________/\_______________________________________________________________
+//==============================================================================================================================
+//                                                    PQ APPROXIMATIONS
+//------------------------------------------------------------------------------------------------------------------------------
+// PQ is very close to x^(1/8). The functions below Use the fast float approximation method to do
+// PQ<~>Gamma2 (4th power and fast 4th root) and PQ<~>Linear (8th power and fast 8th root). Maximum error is ~0.2%.
+//==============================================================================================================================
+// Helpers
+ AF1 Quart(AF1 a) { a = a * a; return a * a;}
+ AF1 Oct(AF1 a) { a = a * a; a = a * a; return a * a; }
+ AF2 Quart(AF2 a) { a = a * a; return a * a; }
+ AF2 Oct(AF2 a) { a = a * a; a = a * a; return a * a; }
+ AF3 Quart(AF3 a) { a = a * a; return a * a; }
+ AF3 Oct(AF3 a) { a = a * a; a = a * a; return a * a; }
+ AF4 Quart(AF4 a) { a = a * a; return a * a; }
+ AF4 Oct(AF4 a) { a = a * a; a = a * a; return a * a; }
+ //------------------------------------------------------------------------------------------------------------------------------
+ AF1 APrxPQToGamma2(AF1 a) { return Quart(a); }
+ AF1 APrxPQToLinear(AF1 a) { return Oct(a); }
+ AF1 APrxLoGamma2ToPQ(AF1 a) { return AF1_AU1((AU1_AF1(a) >> AU1_(2)) + AU1_(0x2F9A4E46)); }
+ AF1 APrxMedGamma2ToPQ(AF1 a) { AF1 b = AF1_AU1((AU1_AF1(a) >> AU1_(2)) + AU1_(0x2F9A4E46)); AF1 b4 = Quart(b); return b - b * (b4 - a) / (AF1_(4.0) * b4); }
+ AF1 APrxHighGamma2ToPQ(AF1 a) { return sqrt(sqrt(a)); }
+ AF1 APrxLoLinearToPQ(AF1 a) { return AF1_AU1((AU1_AF1(a) >> AU1_(3)) + AU1_(0x378D8723)); }
+ AF1 APrxMedLinearToPQ(AF1 a) { AF1 b = AF1_AU1((AU1_AF1(a) >> AU1_(3)) + AU1_(0x378D8723)); AF1 b8 = Oct(b); return b - b * (b8 - a) / (AF1_(8.0) * b8); }
+ AF1 APrxHighLinearToPQ(AF1 a) { return sqrt(sqrt(sqrt(a))); }
+ //------------------------------------------------------------------------------------------------------------------------------
+ AF2 APrxPQToGamma2(AF2 a) { return Quart(a); }
+ AF2 APrxPQToLinear(AF2 a) { return Oct(a); }
+ AF2 APrxLoGamma2ToPQ(AF2 a) { return AF2_AU2((AU2_AF2(a) >> AU2_(2)) + AU2_(0x2F9A4E46)); }
+ AF2 APrxMedGamma2ToPQ(AF2 a) { AF2 b = AF2_AU2((AU2_AF2(a) >> AU2_(2)) + AU2_(0x2F9A4E46)); AF2 b4 = Quart(b); return b - b * (b4 - a) / (AF1_(4.0) * b4); }
+ AF2 APrxHighGamma2ToPQ(AF2 a) { return sqrt(sqrt(a)); }
+ AF2 APrxLoLinearToPQ(AF2 a) { return AF2_AU2((AU2_AF2(a) >> AU2_(3)) + AU2_(0x378D8723)); }
+ AF2 APrxMedLinearToPQ(AF2 a) { AF2 b = AF2_AU2((AU2_AF2(a) >> AU2_(3)) + AU2_(0x378D8723)); AF2 b8 = Oct(b); return b - b * (b8 - a) / (AF1_(8.0) * b8); }
+ AF2 APrxHighLinearToPQ(AF2 a) { return sqrt(sqrt(sqrt(a))); }
+ //------------------------------------------------------------------------------------------------------------------------------
+ AF3 APrxPQToGamma2(AF3 a) { return Quart(a); }
+ AF3 APrxPQToLinear(AF3 a) { return Oct(a); }
+ AF3 APrxLoGamma2ToPQ(AF3 a) { return AF3_AU3((AU3_AF3(a) >> AU3_(2)) + AU3_(0x2F9A4E46)); }
+ AF3 APrxMedGamma2ToPQ(AF3 a) { AF3 b = AF3_AU3((AU3_AF3(a) >> AU3_(2)) + AU3_(0x2F9A4E46)); AF3 b4 = Quart(b); return b - b * (b4 - a) / (AF1_(4.0) * b4); }
+ AF3 APrxHighGamma2ToPQ(AF3 a) { return sqrt(sqrt(a)); }
+ AF3 APrxLoLinearToPQ(AF3 a) { return AF3_AU3((AU3_AF3(a) >> AU3_(3)) + AU3_(0x378D8723)); }
+ AF3 APrxMedLinearToPQ(AF3 a) { AF3 b = AF3_AU3((AU3_AF3(a) >> AU3_(3)) + AU3_(0x378D8723)); AF3 b8 = Oct(b); return b - b * (b8 - a) / (AF1_(8.0) * b8); }
+ AF3 APrxHighLinearToPQ(AF3 a) { return sqrt(sqrt(sqrt(a))); }
+ //------------------------------------------------------------------------------------------------------------------------------
+ AF4 APrxPQToGamma2(AF4 a) { return Quart(a); }
+ AF4 APrxPQToLinear(AF4 a) { return Oct(a); }
+ AF4 APrxLoGamma2ToPQ(AF4 a) { return AF4_AU4((AU4_AF4(a) >> AU4_(2)) + AU4_(0x2F9A4E46)); }
+ AF4 APrxMedGamma2ToPQ(AF4 a) { AF4 b = AF4_AU4((AU4_AF4(a) >> AU4_(2)) + AU4_(0x2F9A4E46)); AF4 b4 = Quart(b); return b - b * (b4 - a) / (AF1_(4.0) * b4); }
+ AF4 APrxHighGamma2ToPQ(AF4 a) { return sqrt(sqrt(a)); }
+ AF4 APrxLoLinearToPQ(AF4 a) { return AF4_AU4((AU4_AF4(a) >> AU4_(3)) + AU4_(0x378D8723)); }
+ AF4 APrxMedLinearToPQ(AF4 a) { AF4 b = AF4_AU4((AU4_AF4(a) >> AU4_(3)) + AU4_(0x378D8723)); AF4 b8 = Oct(b); return b - b * (b8 - a) / (AF1_(8.0) * b8); }
+ AF4 APrxHighLinearToPQ(AF4 a) { return sqrt(sqrt(sqrt(a))); }
+////////////////////////////////////////////////////////////////////////////////////////////////////////////////////////////////
+////////////////////////////////////////////////////////////////////////////////////////////////////////////////////////////////
+//_____________________________________________________________/\_______________________________________________________________
+//==============================================================================================================================
+//                                                    PARABOLIC SIN & COS
+//------------------------------------------------------------------------------------------------------------------------------
+// Approximate answers to transcendental questions.
+//------------------------------------------------------------------------------------------------------------------------------
+//==============================================================================================================================
+ #if 1
+  // Valid input range is {-1 to 1} representing {0 to 2 pi}.
+  // Output range is {-1/4 to 1/4} representing {-1 to 1}.
+  AF1 APSinF1(AF1 x){return x*abs(x)-x;} // MAD.
+  AF2 APSinF2(AF2 x){return x*abs(x)-x;}
+  AF1 APCosF1(AF1 x){x=AFractF1(x*AF1_(0.5)+AF1_(0.75));x=x*AF1_(2.0)-AF1_(1.0);return APSinF1(x);} // 3x MAD, FRACT
+  AF2 APCosF2(AF2 x){x=AFractF2(x*AF2_(0.5)+AF2_(0.75));x=x*AF2_(2.0)-AF2_(1.0);return APSinF2(x);}
+  AF2 APSinCosF1(AF1 x){AF1 y=AFractF1(x*AF1_(0.5)+AF1_(0.75));y=y*AF1_(2.0)-AF1_(1.0);return APSinF2(AF2(x,y));}
+ #endif
+//------------------------------------------------------------------------------------------------------------------------------
+ #ifdef A_HALF
+  // For a packed {sin,cos} pair,
+  //  - Native takes 16 clocks and 4 issue slots (no packed transcendentals).
+  //  - Parabolic takes 8 clocks and 8 issue slots (only fract is non-packed).
+  AH1 APSinH1(AH1 x){return x*abs(x)-x;}
+  AH2 APSinH2(AH2 x){return x*abs(x)-x;} // AND,FMA
+  AH1 APCosH1(AH1 x){x=AFractH1(x*AH1_(0.5)+AH1_(0.75));x=x*AH1_(2.0)-AH1_(1.0);return APSinH1(x);} 
+  AH2 APCosH2(AH2 x){x=AFractH2(x*AH2_(0.5)+AH2_(0.75));x=x*AH2_(2.0)-AH2_(1.0);return APSinH2(x);} // 3x FMA, 2xFRACT, AND
+  AH2 APSinCosH1(AH1 x){AH1 y=AFractH1(x*AH1_(0.5)+AH1_(0.75));y=y*AH1_(2.0)-AH1_(1.0);return APSinH2(AH2(x,y));}
+ #endif
+////////////////////////////////////////////////////////////////////////////////////////////////////////////////////////////////
+////////////////////////////////////////////////////////////////////////////////////////////////////////////////////////////////
+//_____________________________________________________________/\_______________________________________________________________
+//==============================================================================================================================
+//                                                     [ZOL] ZERO ONE LOGIC
+//------------------------------------------------------------------------------------------------------------------------------
+// Conditional free logic designed for easy 16-bit packing, and backwards porting to 32-bit.
+//------------------------------------------------------------------------------------------------------------------------------
+// 0 := false
+// 1 := true
+//------------------------------------------------------------------------------------------------------------------------------
+// AndNot(x,y)   -> !(x&y) .... One op.
+// AndOr(x,y,z)  -> (x&y)|z ... One op.
+// GtZero(x)     -> x>0.0 ..... One op.
+// Sel(x,y,z)    -> x?y:z ..... Two ops, has no precision loss.
+// Signed(x)     -> x<0.0 ..... One op.
+// ZeroPass(x,y) -> x?0:y ..... Two ops, 'y' is a pass through safe for aliasing as integer.
+//------------------------------------------------------------------------------------------------------------------------------
+// OPTIMIZATION NOTES
+// ==================
+// - On Vega to use 2 constants in a packed op, pass in as one AW2 or one AH2 'k.xy' and use as 'k.xx' and 'k.yy'.
+//   For example 'a.xy*k.xx+k.yy'.
+//==============================================================================================================================
+ #if 1
+  AU1 AZolAndU1(AU1 x,AU1 y){return min(x,y);}
+  AU2 AZolAndU2(AU2 x,AU2 y){return min(x,y);}
+  AU3 AZolAndU3(AU3 x,AU3 y){return min(x,y);}
+  AU4 AZolAndU4(AU4 x,AU4 y){return min(x,y);}
+//------------------------------------------------------------------------------------------------------------------------------
+  AU1 AZolNotU1(AU1 x){return x^AU1_(1);}
+  AU2 AZolNotU2(AU2 x){return x^AU2_(1);}
+  AU3 AZolNotU3(AU3 x){return x^AU3_(1);}
+  AU4 AZolNotU4(AU4 x){return x^AU4_(1);}
+//------------------------------------------------------------------------------------------------------------------------------
+  AU1 AZolOrU1(AU1 x,AU1 y){return max(x,y);}
+  AU2 AZolOrU2(AU2 x,AU2 y){return max(x,y);}
+  AU3 AZolOrU3(AU3 x,AU3 y){return max(x,y);}
+  AU4 AZolOrU4(AU4 x,AU4 y){return max(x,y);}
+//==============================================================================================================================
+  AU1 AZolF1ToU1(AF1 x){return AU1(x);}
+  AU2 AZolF2ToU2(AF2 x){return AU2(x);}
+  AU3 AZolF3ToU3(AF3 x){return AU3(x);}
+  AU4 AZolF4ToU4(AF4 x){return AU4(x);}
+//------------------------------------------------------------------------------------------------------------------------------
+  // 2 ops, denormals don't work in 32-bit on PC (and if they are enabled, OMOD is disabled).
+  AU1 AZolNotF1ToU1(AF1 x){return AU1(AF1_(1.0)-x);}
+  AU2 AZolNotF2ToU2(AF2 x){return AU2(AF2_(1.0)-x);}
+  AU3 AZolNotF3ToU3(AF3 x){return AU3(AF3_(1.0)-x);}
+  AU4 AZolNotF4ToU4(AF4 x){return AU4(AF4_(1.0)-x);}
+//------------------------------------------------------------------------------------------------------------------------------
+  AF1 AZolU1ToF1(AU1 x){return AF1(x);}
+  AF2 AZolU2ToF2(AU2 x){return AF2(x);}
+  AF3 AZolU3ToF3(AU3 x){return AF3(x);}
+  AF4 AZolU4ToF4(AU4 x){return AF4(x);}
+//==============================================================================================================================
+  AF1 AZolAndF1(AF1 x,AF1 y){return min(x,y);}
+  AF2 AZolAndF2(AF2 x,AF2 y){return min(x,y);}
+  AF3 AZolAndF3(AF3 x,AF3 y){return min(x,y);}
+  AF4 AZolAndF4(AF4 x,AF4 y){return min(x,y);}
+//------------------------------------------------------------------------------------------------------------------------------
+  AF1 ASolAndNotF1(AF1 x,AF1 y){return (-x)*y+AF1_(1.0);}
+  AF2 ASolAndNotF2(AF2 x,AF2 y){return (-x)*y+AF2_(1.0);}
+  AF3 ASolAndNotF3(AF3 x,AF3 y){return (-x)*y+AF3_(1.0);}
+  AF4 ASolAndNotF4(AF4 x,AF4 y){return (-x)*y+AF4_(1.0);}
+//------------------------------------------------------------------------------------------------------------------------------
+  AF1 AZolAndOrF1(AF1 x,AF1 y,AF1 z){return ASatF1(x*y+z);}
+  AF2 AZolAndOrF2(AF2 x,AF2 y,AF2 z){return ASatF2(x*y+z);}
+  AF3 AZolAndOrF3(AF3 x,AF3 y,AF3 z){return ASatF3(x*y+z);}
+  AF4 AZolAndOrF4(AF4 x,AF4 y,AF4 z){return ASatF4(x*y+z);}
+//------------------------------------------------------------------------------------------------------------------------------
+  AF1 AZolGtZeroF1(AF1 x){return ASatF1(x*AF1_(A_INFP_F));}
+  AF2 AZolGtZeroF2(AF2 x){return ASatF2(x*AF2_(A_INFP_F));}
+  AF3 AZolGtZeroF3(AF3 x){return ASatF3(x*AF3_(A_INFP_F));}
+  AF4 AZolGtZeroF4(AF4 x){return ASatF4(x*AF4_(A_INFP_F));}
+//------------------------------------------------------------------------------------------------------------------------------
+  AF1 AZolNotF1(AF1 x){return AF1_(1.0)-x;}
+  AF2 AZolNotF2(AF2 x){return AF2_(1.0)-x;}
+  AF3 AZolNotF3(AF3 x){return AF3_(1.0)-x;}
+  AF4 AZolNotF4(AF4 x){return AF4_(1.0)-x;}
+//------------------------------------------------------------------------------------------------------------------------------
+  AF1 AZolOrF1(AF1 x,AF1 y){return max(x,y);}
+  AF2 AZolOrF2(AF2 x,AF2 y){return max(x,y);}
+  AF3 AZolOrF3(AF3 x,AF3 y){return max(x,y);}
+  AF4 AZolOrF4(AF4 x,AF4 y){return max(x,y);}
+//------------------------------------------------------------------------------------------------------------------------------
+  AF1 AZolSelF1(AF1 x,AF1 y,AF1 z){AF1 r=(-x)*z+z;return x*y+r;}
+  AF2 AZolSelF2(AF2 x,AF2 y,AF2 z){AF2 r=(-x)*z+z;return x*y+r;}
+  AF3 AZolSelF3(AF3 x,AF3 y,AF3 z){AF3 r=(-x)*z+z;return x*y+r;}
+  AF4 AZolSelF4(AF4 x,AF4 y,AF4 z){AF4 r=(-x)*z+z;return x*y+r;}
+//------------------------------------------------------------------------------------------------------------------------------
+  AF1 AZolSignedF1(AF1 x){return ASatF1(x*AF1_(A_INFN_F));}
+  AF2 AZolSignedF2(AF2 x){return ASatF2(x*AF2_(A_INFN_F));}
+  AF3 AZolSignedF3(AF3 x){return ASatF3(x*AF3_(A_INFN_F));}
+  AF4 AZolSignedF4(AF4 x){return ASatF4(x*AF4_(A_INFN_F));}
+//------------------------------------------------------------------------------------------------------------------------------
+  AF1 AZolZeroPassF1(AF1 x,AF1 y){return AF1_AU1(select((AU1_AF1(x)!=AU1_(0)), AU1_(0), AU1_AF1(y)));}
+  AF2 AZolZeroPassF2(AF2 x,AF2 y){return AF2_AU2(select((AU2_AF2(x)!=AU2_(0)), AU2_(0), AU2_AF2(y)));}
+  AF3 AZolZeroPassF3(AF3 x,AF3 y){return AF3_AU3(select((AU3_AF3(x)!=AU3_(0)), AU3_(0), AU3_AF3(y)));}
+  AF4 AZolZeroPassF4(AF4 x,AF4 y){return AF4_AU4(select((AU4_AF4(x)!=AU4_(0)), AU4_(0), AU4_AF4(y)));}
+ #endif
+//==============================================================================================================================
+ #ifdef A_HALF
+  AW1 AZolAndW1(AW1 x,AW1 y){return min(x,y);}
+  AW2 AZolAndW2(AW2 x,AW2 y){return min(x,y);}
+  AW3 AZolAndW3(AW3 x,AW3 y){return min(x,y);}
+  AW4 AZolAndW4(AW4 x,AW4 y){return min(x,y);}
+//------------------------------------------------------------------------------------------------------------------------------
+  AW1 AZolNotW1(AW1 x){return x^AW1_(1);}
+  AW2 AZolNotW2(AW2 x){return x^AW2_(1);}
+  AW3 AZolNotW3(AW3 x){return x^AW3_(1);}
+  AW4 AZolNotW4(AW4 x){return x^AW4_(1);}
+//------------------------------------------------------------------------------------------------------------------------------
+  AW1 AZolOrW1(AW1 x,AW1 y){return max(x,y);}
+  AW2 AZolOrW2(AW2 x,AW2 y){return max(x,y);}
+  AW3 AZolOrW3(AW3 x,AW3 y){return max(x,y);}
+  AW4 AZolOrW4(AW4 x,AW4 y){return max(x,y);}
+//==============================================================================================================================
+  // Uses denormal trick.
+  AW1 AZolH1ToW1(AH1 x){return AW1_AH1(x*AH1_AW1(AW1_(1)));}
+  AW2 AZolH2ToW2(AH2 x){return AW2_AH2(x*AH2_AW2(AW2_(1)));}
+  AW3 AZolH3ToW3(AH3 x){return AW3_AH3(x*AH3_AW3(AW3_(1)));}
+  AW4 AZolH4ToW4(AH4 x){return AW4_AH4(x*AH4_AW4(AW4_(1)));}
+//------------------------------------------------------------------------------------------------------------------------------
+  // AMD arch lacks a packed conversion opcode.
+  AH1 AZolW1ToH1(AW1 x){return AH1_AW1(x*AW1_AH1(AH1_(1.0)));}
+  AH2 AZolW2ToH2(AW2 x){return AH2_AW2(x*AW2_AH2(AH2_(1.0)));}
+  AH3 AZolW1ToH3(AW3 x){return AH3_AW3(x*AW3_AH3(AH3_(1.0)));}
+  AH4 AZolW2ToH4(AW4 x){return AH4_AW4(x*AW4_AH4(AH4_(1.0)));}
+//==============================================================================================================================
+  AH1 AZolAndH1(AH1 x,AH1 y){return min(x,y);}
+  AH2 AZolAndH2(AH2 x,AH2 y){return min(x,y);}
+  AH3 AZolAndH3(AH3 x,AH3 y){return min(x,y);}
+  AH4 AZolAndH4(AH4 x,AH4 y){return min(x,y);}
+//------------------------------------------------------------------------------------------------------------------------------
+  AH1 ASolAndNotH1(AH1 x,AH1 y){return (-x)*y+AH1_(1.0);}
+  AH2 ASolAndNotH2(AH2 x,AH2 y){return (-x)*y+AH2_(1.0);}
+  AH3 ASolAndNotH3(AH3 x,AH3 y){return (-x)*y+AH3_(1.0);}
+  AH4 ASolAndNotH4(AH4 x,AH4 y){return (-x)*y+AH4_(1.0);}
+//------------------------------------------------------------------------------------------------------------------------------
+  AH1 AZolAndOrH1(AH1 x,AH1 y,AH1 z){return ASatH1(x*y+z);}
+  AH2 AZolAndOrH2(AH2 x,AH2 y,AH2 z){return ASatH2(x*y+z);}
+  AH3 AZolAndOrH3(AH3 x,AH3 y,AH3 z){return ASatH3(x*y+z);}
+  AH4 AZolAndOrH4(AH4 x,AH4 y,AH4 z){return ASatH4(x*y+z);}
+//------------------------------------------------------------------------------------------------------------------------------
+  AH1 AZolGtZeroH1(AH1 x){return ASatH1(x*AH1_(A_INFP_H));}
+  AH2 AZolGtZeroH2(AH2 x){return ASatH2(x*AH2_(A_INFP_H));}
+  AH3 AZolGtZeroH3(AH3 x){return ASatH3(x*AH3_(A_INFP_H));}
+  AH4 AZolGtZeroH4(AH4 x){return ASatH4(x*AH4_(A_INFP_H));}
+//------------------------------------------------------------------------------------------------------------------------------
+  AH1 AZolNotH1(AH1 x){return AH1_(1.0)-x;}
+  AH2 AZolNotH2(AH2 x){return AH2_(1.0)-x;}
+  AH3 AZolNotH3(AH3 x){return AH3_(1.0)-x;}
+  AH4 AZolNotH4(AH4 x){return AH4_(1.0)-x;}
+//------------------------------------------------------------------------------------------------------------------------------
+  AH1 AZolOrH1(AH1 x,AH1 y){return max(x,y);}
+  AH2 AZolOrH2(AH2 x,AH2 y){return max(x,y);}
+  AH3 AZolOrH3(AH3 x,AH3 y){return max(x,y);}
+  AH4 AZolOrH4(AH4 x,AH4 y){return max(x,y);}
+//------------------------------------------------------------------------------------------------------------------------------
+  AH1 AZolSelH1(AH1 x,AH1 y,AH1 z){AH1 r=(-x)*z+z;return x*y+r;}
+  AH2 AZolSelH2(AH2 x,AH2 y,AH2 z){AH2 r=(-x)*z+z;return x*y+r;}
+  AH3 AZolSelH3(AH3 x,AH3 y,AH3 z){AH3 r=(-x)*z+z;return x*y+r;}
+  AH4 AZolSelH4(AH4 x,AH4 y,AH4 z){AH4 r=(-x)*z+z;return x*y+r;}
+//------------------------------------------------------------------------------------------------------------------------------
+  AH1 AZolSignedH1(AH1 x){return ASatH1(x*AH1_(A_INFN_H));}
+  AH2 AZolSignedH2(AH2 x){return ASatH2(x*AH2_(A_INFN_H));}
+  AH3 AZolSignedH3(AH3 x){return ASatH3(x*AH3_(A_INFN_H));}
+  AH4 AZolSignedH4(AH4 x){return ASatH4(x*AH4_(A_INFN_H));}
+ #endif
+////////////////////////////////////////////////////////////////////////////////////////////////////////////////////////////////
+////////////////////////////////////////////////////////////////////////////////////////////////////////////////////////////////
+//_____________________________________________________________/\_______________________________________________________________
+//==============================================================================================================================
+//                                                      COLOR CONVERSIONS
+//------------------------------------------------------------------------------------------------------------------------------
+// These are all linear to/from some other space (where 'linear' has been shortened out of the function name).
+// So 'ToGamma' is 'LinearToGamma', and 'FromGamma' is 'LinearFromGamma'.
+// These are branch free implementations.
+// The AToSrgbF1() function is useful for stores for compute shaders for GPUs without hardware linear->sRGB store conversion.
+//------------------------------------------------------------------------------------------------------------------------------
+// TRANSFER FUNCTIONS
+// ==================
+// 709 ..... Rec709 used for some HDTVs
+// Gamma ... Typically 2.2 for some PC displays, or 2.4-2.5 for CRTs, or 2.2 FreeSync2 native
+// Pq ...... PQ native for HDR10
+// Srgb .... The sRGB output, typical of PC displays, useful for 10-bit output, or storing to 8-bit UNORM without SRGB type
+// Two ..... Gamma 2.0, fastest conversion (useful for intermediate pass approximations)
+// Three ... Gamma 3.0, less fast, but good for HDR.
+//------------------------------------------------------------------------------------------------------------------------------
+// KEEPING TO SPEC
+// ===============
+// Both Rec.709 and sRGB have a linear segment which as spec'ed would intersect the curved segment 2 times.
+//  (a.) For 8-bit sRGB, steps {0 to 10.3} are in the linear region (4% of the encoding range).
+//  (b.) For 8-bit  709, steps {0 to 20.7} are in the linear region (8% of the encoding range).
+// Also there is a slight step in the transition regions.
+// Precision of the coefficients in the spec being the likely cause.
+// Main usage case of the sRGB code is to do the linear->sRGB converstion in a compute shader before store.
+// This is to work around lack of hardware (typically only ROP does the conversion for free).
+// To "correct" the linear segment, would be to introduce error, because hardware decode of sRGB->linear is fixed (and free).
+// So this header keeps with the spec.
+// For linear->sRGB transforms, the linear segment in some respects reduces error, because rounding in that region is linear.
+// Rounding in the curved region in hardware (and fast software code) introduces error due to rounding in non-linear.
+//------------------------------------------------------------------------------------------------------------------------------
+// FOR PQ
+// ======
+// Both input and output is {0.0-1.0}, and where output 1.0 represents 10000.0 cd/m^2.
+// All constants are only specified to FP32 precision.
+// External PQ source reference,
+//  - [REDACTED_PROFILE_URL]/aces-dev/blob/master/transforms/ctl/utilities/ACESlib.Utilities_Color.a1.0.1.ctl
+//------------------------------------------------------------------------------------------------------------------------------
+// PACKED VERSIONS
+// ===============
+// These are the A*H2() functions.
+// There is no PQ functions as FP16 seemed to not have enough precision for the conversion.
+// The remaining functions are "good enough" for 8-bit, and maybe 10-bit if not concerned about a few 1-bit errors.
+// Precision is lowest in the 709 conversion, higher in sRGB, higher still in Two and Gamma (when using 2.2 at least).
+//------------------------------------------------------------------------------------------------------------------------------
+// NOTES
+// =====
+// Could be faster for PQ conversions to be in ALU or a texture lookup depending on usage case.
+//==============================================================================================================================
+ #if 1
+  AF1 ATo709F1(AF1 c){AF3 j=AF3(0.018*4.5,4.5,0.45);AF2 k=AF2(1.099,-0.099);
+   return clamp(j.x  ,c*j.y  ,pow(c,j.z  )*k.x  +k.y  );}
+  AF2 ATo709F2(AF2 c){AF3 j=AF3(0.018*4.5,4.5,0.45);AF2 k=AF2(1.099,-0.099);
+   return clamp(j.xx ,c*j.yy ,pow(c,j.zz )*k.xx +k.yy );}
+  AF3 ATo709F3(AF3 c){AF3 j=AF3(0.018*4.5,4.5,0.45);AF2 k=AF2(1.099,-0.099);
+   return clamp(j.xxx,c*j.yyy,pow(c,j.zzz)*k.xxx+k.yyy);}
+//------------------------------------------------------------------------------------------------------------------------------
+  // Note 'rcpX' is '1/x', where the 'x' is what would be used in AFromGamma().
+  AF1 AToGammaF1(AF1 c,AF1 rcpX){return pow(c,AF1_(rcpX));} 
+  AF2 AToGammaF2(AF2 c,AF1 rcpX){return pow(c,AF2_(rcpX));} 
+  AF3 AToGammaF3(AF3 c,AF1 rcpX){return pow(c,AF3_(rcpX));} 
+//------------------------------------------------------------------------------------------------------------------------------
+  AF1 AToPqF1(AF1 x){AF1 p=pow(x,AF1_(0.159302));
+   return pow((AF1_(0.835938)+AF1_(18.8516)*p)/(AF1_(1.0)+AF1_(18.6875)*p),AF1_(78.8438));}
+  AF2 AToPqF1(AF2 x){AF2 p=pow(x,AF2_(0.159302));
+   return pow((AF2_(0.835938)+AF2_(18.8516)*p)/(AF2_(1.0)+AF2_(18.6875)*p),AF2_(78.8438));}
+  AF3 AToPqF1(AF3 x){AF3 p=pow(x,AF3_(0.159302));
+   return pow((AF3_(0.835938)+AF3_(18.8516)*p)/(AF3_(1.0)+AF3_(18.6875)*p),AF3_(78.8438));}
+//------------------------------------------------------------------------------------------------------------------------------
+  AF1 AToSrgbF1(AF1 c){AF3 j=AF3(0.0031308*12.92,12.92,1.0/2.4);AF2 k=AF2(1.055,-0.055);
+   return clamp(j.x  ,c*j.y  ,pow(c,j.z  )*k.x  +k.y  );}
+  AF2 AToSrgbF2(AF2 c){AF3 j=AF3(0.0031308*12.92,12.92,1.0/2.4);AF2 k=AF2(1.055,-0.055);
+   return clamp(j.xx ,c*j.yy ,pow(c,j.zz )*k.xx +k.yy );}
+  AF3 AToSrgbF3(AF3 c){AF3 j=AF3(0.0031308*12.92,12.92,1.0/2.4);AF2 k=AF2(1.055,-0.055);
+   return clamp(j.xxx,c*j.yyy,pow(c,j.zzz)*k.xxx+k.yyy);}
+//------------------------------------------------------------------------------------------------------------------------------
+  AF1 AToTwoF1(AF1 c){return sqrt(c);}
+  AF2 AToTwoF2(AF2 c){return sqrt(c);}
+  AF3 AToTwoF3(AF3 c){return sqrt(c);}
+//------------------------------------------------------------------------------------------------------------------------------
+  AF1 AToThreeF1(AF1 c){return pow(c,AF1_(1.0/3.0));}
+  AF2 AToThreeF2(AF2 c){return pow(c,AF2_(1.0/3.0));}
+  AF3 AToThreeF3(AF3 c){return pow(c,AF3_(1.0/3.0));}
+ #endif
+//==============================================================================================================================
+ #if 1
+  // Unfortunately median won't work here.
+  AF1 AFrom709F1(AF1 c){AF3 j=AF3(0.081/4.5,1.0/4.5,1.0/0.45);AF2 k=AF2(1.0/1.099,0.099/1.099);
+   return AZolSelF1(AZolSignedF1(c-j.x  ),c*j.y  ,pow(c*k.x  +k.y  ,j.z  ));}
+  AF2 AFrom709F2(AF2 c){AF3 j=AF3(0.081/4.5,1.0/4.5,1.0/0.45);AF2 k=AF2(1.0/1.099,0.099/1.099);
+   return AZolSelF2(AZolSignedF2(c-j.xx ),c*j.yy ,pow(c*k.xx +k.yy ,j.zz ));}
+  AF3 AFrom709F3(AF3 c){AF3 j=AF3(0.081/4.5,1.0/4.5,1.0/0.45);AF2 k=AF2(1.0/1.099,0.099/1.099);
+   return AZolSelF3(AZolSignedF3(c-j.xxx),c*j.yyy,pow(c*k.xxx+k.yyy,j.zzz));}
+//------------------------------------------------------------------------------------------------------------------------------
+  AF1 AFromGammaF1(AF1 c,AF1 x){return pow(c,AF1_(x));} 
+  AF2 AFromGammaF2(AF2 c,AF1 x){return pow(c,AF2_(x));} 
+  AF3 AFromGammaF3(AF3 c,AF1 x){return pow(c,AF3_(x));} 
+//------------------------------------------------------------------------------------------------------------------------------
+  AF1 AFromPqF1(AF1 x){AF1 p=pow(x,AF1_(0.0126833));
+   return pow(ASatF1(p-AF1_(0.835938))/(AF1_(18.8516)-AF1_(18.6875)*p),AF1_(6.27739));}
+  AF2 AFromPqF1(AF2 x){AF2 p=pow(x,AF2_(0.0126833));
+   return pow(ASatF2(p-AF2_(0.835938))/(AF2_(18.8516)-AF2_(18.6875)*p),AF2_(6.27739));}
+  AF3 AFromPqF1(AF3 x){AF3 p=pow(x,AF3_(0.0126833));
+   return pow(ASatF3(p-AF3_(0.835938))/(AF3_(18.8516)-AF3_(18.6875)*p),AF3_(6.27739));}
+//------------------------------------------------------------------------------------------------------------------------------
+  // Unfortunately median won't work here.
+  AF1 AFromSrgbF1(AF1 c){AF3 j=AF3(0.04045/12.92,1.0/12.92,2.4);AF2 k=AF2(1.0/1.055,0.055/1.055);
+   return AZolSelF1(AZolSignedF1(c-j.x  ),c*j.y  ,pow(c*k.x  +k.y  ,j.z  ));}
+  AF2 AFromSrgbF2(AF2 c){AF3 j=AF3(0.04045/12.92,1.0/12.92,2.4);AF2 k=AF2(1.0/1.055,0.055/1.055);
+   return AZolSelF2(AZolSignedF2(c-j.xx ),c*j.yy ,pow(c*k.xx +k.yy ,j.zz ));}
+  AF3 AFromSrgbF3(AF3 c){AF3 j=AF3(0.04045/12.92,1.0/12.92,2.4);AF2 k=AF2(1.0/1.055,0.055/1.055);
+   return AZolSelF3(AZolSignedF3(c-j.xxx),c*j.yyy,pow(c*k.xxx+k.yyy,j.zzz));}
+//------------------------------------------------------------------------------------------------------------------------------
+  AF1 AFromTwoF1(AF1 c){return c*c;}
+  AF2 AFromTwoF2(AF2 c){return c*c;}
+  AF3 AFromTwoF3(AF3 c){return c*c;}
+//------------------------------------------------------------------------------------------------------------------------------
+  AF1 AFromThreeF1(AF1 c){return c*c*c;}
+  AF2 AFromThreeF2(AF2 c){return c*c*c;}
+  AF3 AFromThreeF3(AF3 c){return c*c*c;}
+ #endif
+//==============================================================================================================================
+ #ifdef A_HALF
+  AH1 ATo709H1(AH1 c){AH3 j=AH3(0.018*4.5,4.5,0.45);AH2 k=AH2(1.099,-0.099);
+   return clamp(j.x  ,c*j.y  ,pow(c,j.z  )*k.x  +k.y  );}
+  AH2 ATo709H2(AH2 c){AH3 j=AH3(0.018*4.5,4.5,0.45);AH2 k=AH2(1.099,-0.099);
+   return clamp(j.xx ,c*j.yy ,pow(c,j.zz )*k.xx +k.yy );}
+  AH3 ATo709H3(AH3 c){AH3 j=AH3(0.018*4.5,4.5,0.45);AH2 k=AH2(1.099,-0.099);
+   return clamp(j.xxx,c*j.yyy,pow(c,j.zzz)*k.xxx+k.yyy);}
+//------------------------------------------------------------------------------------------------------------------------------
+  AH1 AToGammaH1(AH1 c,AH1 rcpX){return pow(c,AH1_(rcpX));}
+  AH2 AToGammaH2(AH2 c,AH1 rcpX){return pow(c,AH2_(rcpX));}
+  AH3 AToGammaH3(AH3 c,AH1 rcpX){return pow(c,AH3_(rcpX));}
+//------------------------------------------------------------------------------------------------------------------------------
+  AH1 AToSrgbH1(AH1 c){AH3 j=AH3(0.0031308*12.92,12.92,1.0/2.4);AH2 k=AH2(1.055,-0.055);
+   return clamp(j.x  ,c*j.y  ,pow(c,j.z  )*k.x  +k.y  );}
+  AH2 AToSrgbH2(AH2 c){AH3 j=AH3(0.0031308*12.92,12.92,1.0/2.4);AH2 k=AH2(1.055,-0.055);
+   return clamp(j.xx ,c*j.yy ,pow(c,j.zz )*k.xx +k.yy );}
+  AH3 AToSrgbH3(AH3 c){AH3 j=AH3(0.0031308*12.92,12.92,1.0/2.4);AH2 k=AH2(1.055,-0.055);
+   return clamp(j.xxx,c*j.yyy,pow(c,j.zzz)*k.xxx+k.yyy);}
+//------------------------------------------------------------------------------------------------------------------------------
+  AH1 AToTwoH1(AH1 c){return sqrt(c);}
+  AH2 AToTwoH2(AH2 c){return sqrt(c);}
+  AH3 AToTwoH3(AH3 c){return sqrt(c);}
+//------------------------------------------------------------------------------------------------------------------------------
+  AH1 AToThreeF1(AH1 c){return pow(c,AH1_(1.0/3.0));}
+  AH2 AToThreeF2(AH2 c){return pow(c,AH2_(1.0/3.0));}
+  AH3 AToThreeF3(AH3 c){return pow(c,AH3_(1.0/3.0));}
+ #endif
+//==============================================================================================================================
+ #ifdef A_HALF
+  AH1 AFrom709H1(AH1 c){AH3 j=AH3(0.081/4.5,1.0/4.5,1.0/0.45);AH2 k=AH2(1.0/1.099,0.099/1.099);
+   return AZolSelH1(AZolSignedH1(c-j.x  ),c*j.y  ,pow(c*k.x  +k.y  ,j.z  ));}
+  AH2 AFrom709H2(AH2 c){AH3 j=AH3(0.081/4.5,1.0/4.5,1.0/0.45);AH2 k=AH2(1.0/1.099,0.099/1.099);
+   return AZolSelH2(AZolSignedH2(c-j.xx ),c*j.yy ,pow(c*k.xx +k.yy ,j.zz ));}
+  AH3 AFrom709H3(AH3 c){AH3 j=AH3(0.081/4.5,1.0/4.5,1.0/0.45);AH2 k=AH2(1.0/1.099,0.099/1.099);
+   return AZolSelH3(AZolSignedH3(c-j.xxx),c*j.yyy,pow(c*k.xxx+k.yyy,j.zzz));}
+//------------------------------------------------------------------------------------------------------------------------------
+  AH1 AFromGammaH1(AH1 c,AH1 x){return pow(c,AH1_(x));}
+  AH2 AFromGammaH2(AH2 c,AH1 x){return pow(c,AH2_(x));}
+  AH3 AFromGammaH3(AH3 c,AH1 x){return pow(c,AH3_(x));}
+//------------------------------------------------------------------------------------------------------------------------------
+  AH1 AHromSrgbF1(AH1 c){AH3 j=AH3(0.04045/12.92,1.0/12.92,2.4);AH2 k=AH2(1.0/1.055,0.055/1.055);
+   return AZolSelH1(AZolSignedH1(c-j.x  ),c*j.y  ,pow(c*k.x  +k.y  ,j.z  ));}
+  AH2 AHromSrgbF2(AH2 c){AH3 j=AH3(0.04045/12.92,1.0/12.92,2.4);AH2 k=AH2(1.0/1.055,0.055/1.055);
+   return AZolSelH2(AZolSignedH2(c-j.xx ),c*j.yy ,pow(c*k.xx +k.yy ,j.zz ));}
+  AH3 AHromSrgbF3(AH3 c){AH3 j=AH3(0.04045/12.92,1.0/12.92,2.4);AH2 k=AH2(1.0/1.055,0.055/1.055);
+   return AZolSelH3(AZolSignedH3(c-j.xxx),c*j.yyy,pow(c*k.xxx+k.yyy,j.zzz));}
+//------------------------------------------------------------------------------------------------------------------------------
+  AH1 AFromTwoH1(AH1 c){return c*c;}
+  AH2 AFromTwoH2(AH2 c){return c*c;}
+  AH3 AFromTwoH3(AH3 c){return c*c;}
+//------------------------------------------------------------------------------------------------------------------------------
+  AH1 AFromThreeH1(AH1 c){return c*c*c;}
+  AH2 AFromThreeH2(AH2 c){return c*c*c;}
+  AH3 AFromThreeH3(AH3 c){return c*c*c;}
+ #endif
+////////////////////////////////////////////////////////////////////////////////////////////////////////////////////////////////
+////////////////////////////////////////////////////////////////////////////////////////////////////////////////////////////////
+//_____________________________________________________________/\_______________________________________________________________
+//==============================================================================================================================
+//                                                          CS REMAP
+//==============================================================================================================================
+ // Simple remap 64x1 to 8x8 with rotated 2x2 pixel quads in quad linear.
+ //  543210
+ //  ======
+ //  ..xxx.
+ //  yy...y
+ AU2 ARmp8x8(AU1 a){return AU2(ABfe(a,1u,3u),ABfiM(ABfe(a,3u,3u),a,1u));}
+//==============================================================================================================================
+ // More complex remap 64x1 to 8x8 which is necessary for 2D wave reductions.
+ //  543210
+ //  ======
+ //  .xx..x
+ //  y..yy.
+ // Details,
+ //  LANE TO 8x8 MAPPING
+ //  ===================
+ //  00 01 08 09 10 11 18 19 
+ //  02 03 0a 0b 12 13 1a 1b
+ //  04 05 0c 0d 14 15 1c 1d
+ //  06 07 0e 0f 16 17 1e 1f 
+ //  20 21 28 29 30 31 38 39 
+ //  22 23 2a 2b 32 33 3a 3b
+ //  24 25 2c 2d 34 35 3c 3d
+ //  26 27 2e 2f 36 37 3e 3f 
+ AU2 ARmpRed8x8(AU1 a){return AU2(ABfiM(ABfe(a,2u,3u),a,1u),ABfiM(ABfe(a,3u,3u),ABfe(a,1u,2u),2u));}
+//==============================================================================================================================
+ #ifdef A_HALF
+  AW2 ARmp8x8H(AU1 a){return AW2(ABfe(a,1u,3u),ABfiM(ABfe(a,3u,3u),a,1u));}
+  AW2 ARmpRed8x8H(AU1 a){return AW2(ABfiM(ABfe(a,2u,3u),a,1u),ABfiM(ABfe(a,3u,3u),ABfe(a,1u,2u),2u));}
+ #endif
+#endif
+////////////////////////////////////////////////////////////////////////////////////////////////////////////////////////////////
+////////////////////////////////////////////////////////////////////////////////////////////////////////////////////////////////
+////////////////////////////////////////////////////////////////////////////////////////////////////////////////////////////////
+////////////////////////////////////////////////////////////////////////////////////////////////////////////////////////////////
+//_____________________________________________________________/\_______________________________________________________________
+//==============================================================================================================================
+//
+//                                                          REFERENCE
+//
+//------------------------------------------------------------------------------------------------------------------------------
+// IEEE FLOAT RULES
+// ================
+//  - saturate(NaN)=0, saturate(-INF)=0, saturate(+INF)=1
+//  - {+/-}0 * {+/-}INF = NaN
+//  - -INF + (+INF) = NaN
+//  - {+/-}0 / {+/-}0 = NaN
+//  - {+/-}INF / {+/-}INF = NaN
+//  - a<(-0) := sqrt(a) = NaN (a=-0.0 won't NaN)
+//  - 0 == -0
+//  - 4/0 = +INF
+//  - 4/-0 = -INF
+//  - 4+INF = +INF
+//  - 4-INF = -INF
+//  - 4*(+INF) = +INF
+//  - 4*(-INF) = -INF
+//  - -4*(+INF) = -INF
+//  - sqrt(+INF) = +INF
+//------------------------------------------------------------------------------------------------------------------------------
+// FP16 ENCODING
+// =============
+// fedcba9876543210
+// ----------------
+// ......mmmmmmmmmm  10-bit mantissa (encodes 11-bit 0.5 to 1.0 except for denormals)
+// .eeeee..........  5-bit exponent
+// .00000..........  denormals
+// .00001..........  -14 exponent
+// .11110..........   15 exponent
+// .111110000000000  infinity
+// .11111nnnnnnnnnn  NaN with n!=0
+// s...............  sign
+//------------------------------------------------------------------------------------------------------------------------------
+// FP16/INT16 ALIASING DENORMAL
+// ============================
+// 11-bit unsigned integers alias with half float denormal/normal values,
+//     1 = 2^(-24) = 1/16777216 ....................... first denormal value
+//     2 = 2^(-23)
+//   ...
+//  1023 = 2^(-14)*(1-2^(-10)) = 2^(-14)*(1-1/1024) ... last denormal value
+//  1024 = 2^(-14) = 1/16384 .......................... first normal value that still maps to integers
+//  2047 .............................................. last normal value that still maps to integers 
+// Scaling limits,
+//  2^15 = 32768 ...................................... largest power of 2 scaling
+// Largest pow2 conversion mapping is at *32768,
+//     1 : 2^(-9) = 1/512
+//     2 : 1/256
+//     4 : 1/128
+//     8 : 1/64
+//    16 : 1/32
+//    32 : 1/16
+//    64 : 1/8
+//   128 : 1/4
+//   256 : 1/2
+//   512 : 1
+//  1024 : 2
+//  2047 : a little less than 4
+//==============================================================================================================================
+////////////////////////////////////////////////////////////////////////////////////////////////////////////////////////////////
+////////////////////////////////////////////////////////////////////////////////////////////////////////////////////////////////
+////////////////////////////////////////////////////////////////////////////////////////////////////////////////////////////////
+////////////////////////////////////////////////////////////////////////////////////////////////////////////////////////////////
+////////////////////////////////////////////////////////////////////////////////////////////////////////////////////////////////
+////////////////////////////////////////////////////////////////////////////////////////////////////////////////////////////////
+//_____________________________________________________________/\_______________________________________________________________
+//==============================================================================================================================
+//
+//
+//                                                     GPU/CPU PORTABILITY
+//
+//
+//------------------------------------------------------------------------------------------------------------------------------
+// This is the GPU implementation.
+// See the CPU implementation for docs.
+//==============================================================================================================================
+#ifdef A_GPU
+ #define A_TRUE true
+ #define A_FALSE false
+ #define A_STATIC
+////////////////////////////////////////////////////////////////////////////////////////////////////////////////////////////////
+////////////////////////////////////////////////////////////////////////////////////////////////////////////////////////////////
+//_____________________________________________________________/\_______________________________________________________________
+//==============================================================================================================================
+//                                     VECTOR ARGUMENT/RETURN/INITIALIZATION PORTABILITY
+//==============================================================================================================================
+ #define retAD2 AD2
+ #define retAD3 AD3
+ #define retAD4 AD4
+ #define retAF2 AF2
+ #define retAF3 AF3
+ #define retAF4 AF4
+ #define retAL2 AL2
+ #define retAL3 AL3
+ #define retAL4 AL4
+ #define retAU2 AU2
+ #define retAU3 AU3
+ #define retAU4 AU4
+//------------------------------------------------------------------------------------------------------------------------------
+ #define inAD2 in AD2
+ #define inAD3 in AD3
+ #define inAD4 in AD4
+ #define inAF2 in AF2
+ #define inAF3 in AF3
+ #define inAF4 in AF4
+ #define inAL2 in AL2
+ #define inAL3 in AL3
+ #define inAL4 in AL4
+ #define inAU2 in AU2
+ #define inAU3 in AU3
+ #define inAU4 in AU4
+//------------------------------------------------------------------------------------------------------------------------------
+ #define inoutAD2 inout AD2
+ #define inoutAD3 inout AD3
+ #define inoutAD4 inout AD4
+ #define inoutAF2 inout AF2
+ #define inoutAF3 inout AF3
+ #define inoutAF4 inout AF4
+ #define inoutAL2 inout AL2
+ #define inoutAL3 inout AL3
+ #define inoutAL4 inout AL4
+ #define inoutAU2 inout AU2
+ #define inoutAU3 inout AU3
+ #define inoutAU4 inout AU4
+//------------------------------------------------------------------------------------------------------------------------------
+ #define outAD2 out AD2
+ #define outAD3 out AD3
+ #define outAD4 out AD4
+ #define outAF2 out AF2
+ #define outAF3 out AF3
+ #define outAF4 out AF4
+ #define outAL2 out AL2
+ #define outAL3 out AL3
+ #define outAL4 out AL4
+ #define outAU2 out AU2
+ #define outAU3 out AU3
+ #define outAU4 out AU4
+//------------------------------------------------------------------------------------------------------------------------------
+ #define varAD2(x) AD2 x
+ #define varAD3(x) AD3 x
+ #define varAD4(x) AD4 x
+ #define varAF2(x) AF2 x
+ #define varAF3(x) AF3 x
+ #define varAF4(x) AF4 x
+ #define varAL2(x) AL2 x
+ #define varAL3(x) AL3 x
+ #define varAL4(x) AL4 x
+ #define varAU2(x) AU2 x
+ #define varAU3(x) AU3 x
+ #define varAU4(x) AU4 x
+//------------------------------------------------------------------------------------------------------------------------------
+ #define initAD2(x,y) AD2(x,y)
+ #define initAD3(x,y,z) AD3(x,y,z)
+ #define initAD4(x,y,z,w) AD4(x,y,z,w)
+ #define initAF2(x,y) AF2(x,y)
+ #define initAF3(x,y,z) AF3(x,y,z)
+ #define initAF4(x,y,z,w) AF4(x,y,z,w)
+ #define initAL2(x,y) AL2(x,y)
+ #define initAL3(x,y,z) AL3(x,y,z)
+ #define initAL4(x,y,z,w) AL4(x,y,z,w)
+ #define initAU2(x,y) AU2(x,y)
+ #define initAU3(x,y,z) AU3(x,y,z)
+ #define initAU4(x,y,z,w) AU4(x,y,z,w)
+////////////////////////////////////////////////////////////////////////////////////////////////////////////////////////////////
+////////////////////////////////////////////////////////////////////////////////////////////////////////////////////////////////
+//_____________________________________________________________/\_______________________________________________________________
+//==============================================================================================================================
+//                                                     SCALAR RETURN OPS
+//==============================================================================================================================
+ #define AAbsD1(a) abs(AD1(a))
+ #define AAbsF1(a) abs(AF1(a))
+//------------------------------------------------------------------------------------------------------------------------------
+ #define ACosD1(a) cos(AD1(a))
+ #define ACosF1(a) cos(AF1(a))
+//------------------------------------------------------------------------------------------------------------------------------
+ #define ADotD2(a,b) dot(AD2(a),AD2(b))
+ #define ADotD3(a,b) dot(AD3(a),AD3(b))
+ #define ADotD4(a,b) dot(AD4(a),AD4(b))
+ #define ADotF2(a,b) dot(AF2(a),AF2(b))
+ #define ADotF3(a,b) dot(AF3(a),AF3(b))
+ #define ADotF4(a,b) dot(AF4(a),AF4(b))
+//------------------------------------------------------------------------------------------------------------------------------
+ #define AExp2D1(a) exp2(AD1(a))
+ #define AExp2F1(a) exp2(AF1(a))
+//------------------------------------------------------------------------------------------------------------------------------
+ #define AFloorD1(a) floor(AD1(a))
+ #define AFloorF1(a) floor(AF1(a))
+//------------------------------------------------------------------------------------------------------------------------------
+ #define ALog2D1(a) log2(AD1(a))
+ #define ALog2F1(a) log2(AF1(a))
+//------------------------------------------------------------------------------------------------------------------------------
+ #define AMaxD1(a,b) max(a,b)
+ #define AMaxF1(a,b) max(a,b)
+ #define AMaxL1(a,b) max(a,b)
+ #define AMaxU1(a,b) max(a,b)
+//------------------------------------------------------------------------------------------------------------------------------
+ #define AMinD1(a,b) min(a,b)
+ #define AMinF1(a,b) min(a,b)
+ #define AMinL1(a,b) min(a,b)
+ #define AMinU1(a,b) min(a,b)
+//------------------------------------------------------------------------------------------------------------------------------
+ #define ASinD1(a) sin(AD1(a))
+ #define ASinF1(a) sin(AF1(a))
+//------------------------------------------------------------------------------------------------------------------------------
+ #define ASqrtD1(a) sqrt(AD1(a))
+ #define ASqrtF1(a) sqrt(AF1(a))
+////////////////////////////////////////////////////////////////////////////////////////////////////////////////////////////////
+////////////////////////////////////////////////////////////////////////////////////////////////////////////////////////////////
+//_____________________________________________________________/\_______________________________________________________________
+//==============================================================================================================================
+//                                               SCALAR RETURN OPS - DEPENDENT
+//==============================================================================================================================
+ #define APowD1(a,b) pow(AD1(a),AF1(b))
+ #define APowF1(a,b) pow(AF1(a),AF1(b))
+////////////////////////////////////////////////////////////////////////////////////////////////////////////////////////////////
+////////////////////////////////////////////////////////////////////////////////////////////////////////////////////////////////
+//_____________________________________________________________/\_______________________________________________________________
+//==============================================================================================================================
+//                                                         VECTOR OPS
+//------------------------------------------------------------------------------------------------------------------------------
+// These are added as needed for production or prototyping, so not necessarily a complete set.
+// They follow a convention of taking in a destination and also returning the destination value to increase utility.
+//==============================================================================================================================
+ #ifdef A_DUBL
+  AD2 opAAbsD2(outAD2 d,inAD2 a){d=abs(a);return d;}
+  AD3 opAAbsD3(outAD3 d,inAD3 a){d=abs(a);return d;}
+  AD4 opAAbsD4(outAD4 d,inAD4 a){d=abs(a);return d;}
+//------------------------------------------------------------------------------------------------------------------------------
+  AD2 opAAddD2(outAD2 d,inAD2 a,inAD2 b){d=a+b;return d;}
+  AD3 opAAddD3(outAD3 d,inAD3 a,inAD3 b){d=a+b;return d;}
+  AD4 opAAddD4(outAD4 d,inAD4 a,inAD4 b){d=a+b;return d;}
+//------------------------------------------------------------------------------------------------------------------------------
+  AD2 opAAddOneD2(outAD2 d,inAD2 a,AD1 b){d=a+AD2_(b);return d;}
+  AD3 opAAddOneD3(outAD3 d,inAD3 a,AD1 b){d=a+AD3_(b);return d;}
+  AD4 opAAddOneD4(outAD4 d,inAD4 a,AD1 b){d=a+AD4_(b);return d;}
+//------------------------------------------------------------------------------------------------------------------------------
+  AD2 opACpyD2(outAD2 d,inAD2 a){d=a;return d;}
+  AD3 opACpyD3(outAD3 d,inAD3 a){d=a;return d;}
+  AD4 opACpyD4(outAD4 d,inAD4 a){d=a;return d;}
+//------------------------------------------------------------------------------------------------------------------------------
+  AD2 opALerpD2(outAD2 d,inAD2 a,inAD2 b,inAD2 c){d=ALerpD2(a,b,c);return d;}
+  AD3 opALerpD3(outAD3 d,inAD3 a,inAD3 b,inAD3 c){d=ALerpD3(a,b,c);return d;}
+  AD4 opALerpD4(outAD4 d,inAD4 a,inAD4 b,inAD4 c){d=ALerpD4(a,b,c);return d;}
+//------------------------------------------------------------------------------------------------------------------------------
+  AD2 opALerpOneD2(outAD2 d,inAD2 a,inAD2 b,AD1 c){d=ALerpD2(a,b,AD2_(c));return d;}
+  AD3 opALerpOneD3(outAD3 d,inAD3 a,inAD3 b,AD1 c){d=ALerpD3(a,b,AD3_(c));return d;}
+  AD4 opALerpOneD4(outAD4 d,inAD4 a,inAD4 b,AD1 c){d=ALerpD4(a,b,AD4_(c));return d;}
+//------------------------------------------------------------------------------------------------------------------------------
+  AD2 opAMaxD2(outAD2 d,inAD2 a,inAD2 b){d=max(a,b);return d;}
+  AD3 opAMaxD3(outAD3 d,inAD3 a,inAD3 b){d=max(a,b);return d;}
+  AD4 opAMaxD4(outAD4 d,inAD4 a,inAD4 b){d=max(a,b);return d;}
+//------------------------------------------------------------------------------------------------------------------------------
+  AD2 opAMinD2(outAD2 d,inAD2 a,inAD2 b){d=min(a,b);return d;}
+  AD3 opAMinD3(outAD3 d,inAD3 a,inAD3 b){d=min(a,b);return d;}
+  AD4 opAMinD4(outAD4 d,inAD4 a,inAD4 b){d=min(a,b);return d;}
+//------------------------------------------------------------------------------------------------------------------------------
+  AD2 opAMulD2(outAD2 d,inAD2 a,inAD2 b){d=a*b;return d;}
+  AD3 opAMulD3(outAD3 d,inAD3 a,inAD3 b){d=a*b;return d;}
+  AD4 opAMulD4(outAD4 d,inAD4 a,inAD4 b){d=a*b;return d;}
+//------------------------------------------------------------------------------------------------------------------------------
+  AD2 opAMulOneD2(outAD2 d,inAD2 a,AD1 b){d=a*AD2_(b);return d;}
+  AD3 opAMulOneD3(outAD3 d,inAD3 a,AD1 b){d=a*AD3_(b);return d;}
+  AD4 opAMulOneD4(outAD4 d,inAD4 a,AD1 b){d=a*AD4_(b);return d;}
+//------------------------------------------------------------------------------------------------------------------------------
+  AD2 opANegD2(outAD2 d,inAD2 a){d=-a;return d;}
+  AD3 opANegD3(outAD3 d,inAD3 a){d=-a;return d;}
+  AD4 opANegD4(outAD4 d,inAD4 a){d=-a;return d;}
+//------------------------------------------------------------------------------------------------------------------------------
+  AD2 opARcpD2(outAD2 d,inAD2 a){d=ARcpD2(a);return d;}
+  AD3 opARcpD3(outAD3 d,inAD3 a){d=ARcpD3(a);return d;}
+  AD4 opARcpD4(outAD4 d,inAD4 a){d=ARcpD4(a);return d;}
+ #endif
+//==============================================================================================================================
+ AF2 opAAbsF2(outAF2 d,inAF2 a){d=abs(a);return d;}
+ AF3 opAAbsF3(outAF3 d,inAF3 a){d=abs(a);return d;}
+ AF4 opAAbsF4(outAF4 d,inAF4 a){d=abs(a);return d;}
+//------------------------------------------------------------------------------------------------------------------------------
+ AF2 opAAddF2(outAF2 d,inAF2 a,inAF2 b){d=a+b;return d;}
+ AF3 opAAddF3(outAF3 d,inAF3 a,inAF3 b){d=a+b;return d;}
+ AF4 opAAddF4(outAF4 d,inAF4 a,inAF4 b){d=a+b;return d;}
+//------------------------------------------------------------------------------------------------------------------------------
+ AF2 opAAddOneF2(outAF2 d,inAF2 a,AF1 b){d=a+AF2_(b);return d;}
+ AF3 opAAddOneF3(outAF3 d,inAF3 a,AF1 b){d=a+AF3_(b);return d;}
+ AF4 opAAddOneF4(outAF4 d,inAF4 a,AF1 b){d=a+AF4_(b);return d;}
+//------------------------------------------------------------------------------------------------------------------------------
+ AF2 opACpyF2(outAF2 d,inAF2 a){d=a;return d;}
+ AF3 opACpyF3(outAF3 d,inAF3 a){d=a;return d;}
+ AF4 opACpyF4(outAF4 d,inAF4 a){d=a;return d;}
+//------------------------------------------------------------------------------------------------------------------------------
+ AF2 opALerpF2(outAF2 d,inAF2 a,inAF2 b,inAF2 c){d=ALerpF2(a,b,c);return d;}
+ AF3 opALerpF3(outAF3 d,inAF3 a,inAF3 b,inAF3 c){d=ALerpF3(a,b,c);return d;}
+ AF4 opALerpF4(outAF4 d,inAF4 a,inAF4 b,inAF4 c){d=ALerpF4(a,b,c);return d;}
+//------------------------------------------------------------------------------------------------------------------------------
+ AF2 opALerpOneF2(outAF2 d,inAF2 a,inAF2 b,AF1 c){d=ALerpF2(a,b,AF2_(c));return d;}
+ AF3 opALerpOneF3(outAF3 d,inAF3 a,inAF3 b,AF1 c){d=ALerpF3(a,b,AF3_(c));return d;}
+ AF4 opALerpOneF4(outAF4 d,inAF4 a,inAF4 b,AF1 c){d=ALerpF4(a,b,AF4_(c));return d;}
+//------------------------------------------------------------------------------------------------------------------------------
+ AF2 opAMaxF2(outAF2 d,inAF2 a,inAF2 b){d=max(a,b);return d;}
+ AF3 opAMaxF3(outAF3 d,inAF3 a,inAF3 b){d=max(a,b);return d;}
+ AF4 opAMaxF4(outAF4 d,inAF4 a,inAF4 b){d=max(a,b);return d;}
+//------------------------------------------------------------------------------------------------------------------------------
+ AF2 opAMinF2(outAF2 d,inAF2 a,inAF2 b){d=min(a,b);return d;}
+ AF3 opAMinF3(outAF3 d,inAF3 a,inAF3 b){d=min(a,b);return d;}
+ AF4 opAMinF4(outAF4 d,inAF4 a,inAF4 b){d=min(a,b);return d;}
+//------------------------------------------------------------------------------------------------------------------------------
+ AF2 opAMulF2(outAF2 d,inAF2 a,inAF2 b){d=a*b;return d;}
+ AF3 opAMulF3(outAF3 d,inAF3 a,inAF3 b){d=a*b;return d;}
+ AF4 opAMulF4(outAF4 d,inAF4 a,inAF4 b){d=a*b;return d;}
+//------------------------------------------------------------------------------------------------------------------------------
+ AF2 opAMulOneF2(outAF2 d,inAF2 a,AF1 b){d=a*AF2_(b);return d;}
+ AF3 opAMulOneF3(outAF3 d,inAF3 a,AF1 b){d=a*AF3_(b);return d;}
+ AF4 opAMulOneF4(outAF4 d,inAF4 a,AF1 b){d=a*AF4_(b);return d;}
+//------------------------------------------------------------------------------------------------------------------------------
+ AF2 opANegF2(outAF2 d,inAF2 a){d=-a;return d;}
+ AF3 opANegF3(outAF3 d,inAF3 a){d=-a;return d;}
+ AF4 opANegF4(outAF4 d,inAF4 a){d=-a;return d;}
+//------------------------------------------------------------------------------------------------------------------------------
+ AF2 opARcpF2(outAF2 d,inAF2 a){d=ARcpF2(a);return d;}
+ AF3 opARcpF3(outAF3 d,inAF3 a){d=ARcpF3(a);return d;}
+ AF4 opARcpF4(outAF4 d,inAF4 a){d=ARcpF4(a);return d;}
+#endif
