Removed Mobile FSR from Source/Shaders

diff --git a/Engine/Plugins/Runtime/MobileFSR/Shaders/Private/ffx_fsr1.ush b/Engine/Plugins/Runtime/MobileFSR/Shaders/Private/ffx_fsr1.ush
new file mode 100644
--- /dev/null
+++ b/Engine/Plugins/Runtime/MobileFSR/Shaders/Private/ffx_fsr1.ush
@@ -0,0 +1,1277 @@
+//_____________________________________________________________/\_______________________________________________________________
+//==============================================================================================================================
+//
+//
+//                    AMD FidelityFX SUPER RESOLUTION [FSR 1] ::: SPATIAL SCALING & EXTRAS - v1.20210629
+//
+//
+//------------------------------------------------------------------------------------------------------------------------------
+////////////////////////////////////////////////////////////////////////////////////////////////////////////////////////////////
+////////////////////////////////////////////////////////////////////////////////////////////////////////////////////////////////
+//------------------------------------------------------------------------------------------------------------------------------
+// FidelityFX Super Resolution Sample
+//
+// Copyright (c) 2021 Advanced Micro Devices, Inc. All rights reserved.
+// Permission is hereby granted, free of charge, to any person obtaining a copy
+// of this software and associated documentation files(the "Software"), to deal
+// in the Software without restriction, including without limitation the rights
+// to use, copy, modify, merge, publish, distribute, sublicense, and / or sell
+// copies of the Software, and to permit persons to whom the Software is
+// furnished to do so, subject to the following conditions :
+// The above copyright notice and this permission notice shall be included in
+// all copies or substantial portions of the Software.
+// THE SOFTWARE IS PROVIDED "AS IS", WITHOUT WARRANTY OF ANY KIND, EXPRESS OR
+// IMPLIED, INCLUDING BUT NOT LIMITED TO THE WARRANTIES OF MERCHANTABILITY,
+// FITNESS FOR A PARTICULAR PURPOSE AND NONINFRINGEMENT.IN NO EVENT SHALL THE
+// AUTHORS OR COPYRIGHT HOLDERS BE LIABLE FOR ANY CLAIM, DAMAGES OR OTHER
+// LIABILITY, WHETHER IN AN ACTION OF CONTRACT, TORT OR OTHERWISE, ARISING FROM,
+// OUT OF OR IN CONNECTION WITH THE SOFTWARE OR THE USE OR OTHER DEALINGS IN
+// THE SOFTWARE.
+//------------------------------------------------------------------------------------------------------------------------------
+////////////////////////////////////////////////////////////////////////////////////////////////////////////////////////////////
+////////////////////////////////////////////////////////////////////////////////////////////////////////////////////////////////
+//------------------------------------------------------------------------------------------------------------------------------
+// ABOUT
+// =====
+// FSR is a collection of algorithms relating to generating a higher resolution image.
+// This specific header focuses on single-image non-temporal image scaling, and related tools.
+// 
+// The core functions are EASU and RCAS:
+//  [EASU] Edge Adaptive Spatial Upsampling ....... 1x to 4x area range spatial scaling, clamped adaptive elliptical filter.
+//  [RCAS] Robust Contrast Adaptive Sharpening .... A non-scaling variation on CAS.
+// RCAS needs to be applied after EASU as a separate pass.
+// 
+// Optional utility functions are:
+//  [LFGA] Linear Film Grain Applicator ........... Tool to apply film grain after scaling.
+//  [SRTM] Simple Reversible Tone-Mapper .......... Linear HDR {0 to FP16_MAX} to {0 to 1} and back.
+//  [TEPD] Temporal Energy Preserving Dither ...... Temporally energy preserving dithered {0 to 1} linear to gamma 2.0 conversion.
+// See each individual sub-section for inline documentation.
+//------------------------------------------------------------------------------------------------------------------------------
+////////////////////////////////////////////////////////////////////////////////////////////////////////////////////////////////
+////////////////////////////////////////////////////////////////////////////////////////////////////////////////////////////////
+//------------------------------------------------------------------------------------------------------------------------------
+// FUNCTION PERMUTATIONS
+// =====================
+// *F() ..... Single item computation with 32-bit.
+// *H() ..... Single item computation with 16-bit, with packing (aka two 16-bit ops in parallel) when possible.
+// *Hx2() ... Processing two items in parallel with 16-bit, easier packing.
+//            Not all interfaces in this file have a *Hx2() form.
+//==============================================================================================================================
+////////////////////////////////////////////////////////////////////////////////////////////////////////////////////////////////
+////////////////////////////////////////////////////////////////////////////////////////////////////////////////////////////////
+////////////////////////////////////////////////////////////////////////////////////////////////////////////////////////////////
+////////////////////////////////////////////////////////////////////////////////////////////////////////////////////////////////
+//_____________________________________________________________/\_______________________________________________________________
+//==============================================================================================================================
+//
+//                                        FSR - [EASU] EDGE ADAPTIVE SPATIAL UPSAMPLING
+//
+//------------------------------------------------------------------------------------------------------------------------------
+// EASU provides a high quality spatial-only scaling at relatively low cost.
+// Meaning EASU is appropiate for laptops and other low-end GPUs.
+// Quality from 1x to 4x area scaling is good.
+//------------------------------------------------------------------------------------------------------------------------------
+// The scalar uses a modified fast approximation to the standard lanczos(size=2) kernel.
+// EASU runs in a single pass, so it applies a directionally and anisotropically adaptive radial lanczos.
+// This is also kept as simple as possible to have minimum runtime.
+//------------------------------------------------------------------------------------------------------------------------------
+// The lanzcos filter has negative lobes, so by itself it will introduce ringing.
+// To remove all ringing, the algorithm uses the nearest 2x2 input texels as a neighborhood,
+// and limits output to the minimum and maximum of that neighborhood.
+//------------------------------------------------------------------------------------------------------------------------------
+// Input image requirements:
+// 
+// Color needs to be encoded as 3 channel[red, green, blue](e.g.XYZ not supported)
+// Each channel needs to be in the range[0, 1]
+// Any color primaries are supported
+// Display / tonemapping curve needs to be as if presenting to sRGB display or similar(e.g.Gamma 2.0)
+// There should be no banding in the input
+// There should be no high amplitude noise in the input
+// There should be no noise in the input that is not at input pixel granularity
+// For performance purposes, use 32bpp formats
+//------------------------------------------------------------------------------------------------------------------------------
+// Best to apply EASU at the end of the frame after tonemapping 
+// but before film grain or composite of the UI.
+//------------------------------------------------------------------------------------------------------------------------------
+//==============================================================================================================================
+////////////////////////////////////////////////////////////////////////////////////////////////////////////////////////////////
+////////////////////////////////////////////////////////////////////////////////////////////////////////////////////////////////
+//_____________________________________________________________/\_______________________________________________________________
+//==============================================================================================================================
+//                                                      CONSTANT SETUP
+//==============================================================================================================================
+// Call to setup required constant values (works on CPU or GPU).
+A_STATIC void FsrEasuCon(
+outAU4 con0,
+outAU4 con1,
+outAU4 con2,
+outAU4 con3,
+// This the rendered image resolution being upscaled
+AF1 inputViewportInPixelsX,
+AF1 inputViewportInPixelsY,
+// This is the resolution of the resource containing the input image (useful for dynamic resolution)
+AF1 inputSizeInPixelsX,
+AF1 inputSizeInPixelsY,
+// This is the display resolution which the input image gets upscaled to
+AF1 outputSizeInPixelsX,
+AF1 outputSizeInPixelsY){
+ // Output integer position to a pixel position in viewport.
+ con0[0]=AU1_AF1(inputViewportInPixelsX*ARcpF1(outputSizeInPixelsX));
+ con0[1]=AU1_AF1(inputViewportInPixelsY*ARcpF1(outputSizeInPixelsY));
+ con0[2]=AU1_AF1(AF1_(0.5)*inputViewportInPixelsX*ARcpF1(outputSizeInPixelsX)-AF1_(0.5));
+ con0[3]=AU1_AF1(AF1_(0.5)*inputViewportInPixelsY*ARcpF1(outputSizeInPixelsY)-AF1_(0.5));
+ // Viewport pixel position to normalized image space.
+ // This is used to get upper-left of 'F' tap.
+ con1[0]=AU1_AF1(ARcpF1(inputSizeInPixelsX));
+ con1[1]=AU1_AF1(ARcpF1(inputSizeInPixelsY));
+ // Centers of gather4, first offset from upper-left of 'F'.
+ //      +---+---+
+ //      |   |   |
+ //      +--(0)--+
+ //      | b | c |
+ //  +---F---+---+---+
+ //  | e | f | g | h |
+ //  +--(1)--+--(2)--+
+ //  | i | j | k | l |
+ //  +---+---+---+---+
+ //      | n | o |
+ //      +--(3)--+
+ //      |   |   |
+ //      +---+---+
+ con1[2]=AU1_AF1(AF1_( 1.0)*ARcpF1(inputSizeInPixelsX));
+ con1[3]=AU1_AF1(AF1_(-1.0)*ARcpF1(inputSizeInPixelsY));
+ // These are from (0) instead of 'F'.
+ con2[0]=AU1_AF1(AF1_(-1.0)*ARcpF1(inputSizeInPixelsX));
+ con2[1]=AU1_AF1(AF1_( 2.0)*ARcpF1(inputSizeInPixelsY));
+ con2[2]=AU1_AF1(AF1_( 1.0)*ARcpF1(inputSizeInPixelsX));
+ con2[3]=AU1_AF1(AF1_( 2.0)*ARcpF1(inputSizeInPixelsY));
+ con3[0]=AU1_AF1(AF1_( 0.0)*ARcpF1(inputSizeInPixelsX));
+ con3[1]=AU1_AF1(AF1_( 4.0)*ARcpF1(inputSizeInPixelsY));
+ con3[2]=con3[3]=0;}
+
+//If the an offset into the input image resource
+A_STATIC void FsrEasuConOffset(
+    outAU4 con0,
+    outAU4 con1,
+    outAU4 con2,
+    outAU4 con3,
+    // This the rendered image resolution being upscaled
+    AF1 inputViewportInPixelsX,
+    AF1 inputViewportInPixelsY,
+    // This is the resolution of the resource containing the input image (useful for dynamic resolution)
+    AF1 inputSizeInPixelsX,
+    AF1 inputSizeInPixelsY,
+    // This is the display resolution which the input image gets upscaled to
+    AF1 outputSizeInPixelsX,
+    AF1 outputSizeInPixelsY,
+    // This is the input image offset into the resource containing it (useful for dynamic resolution)
+    AF1 inputOffsetInPixelsX,
+    AF1 inputOffsetInPixelsY) {
+    FsrEasuCon(con0, con1, con2, con3, inputViewportInPixelsX, inputViewportInPixelsY, inputSizeInPixelsX, inputSizeInPixelsY, outputSizeInPixelsX, outputSizeInPixelsY);
+    con0[2] = AU1_AF1(AF1_(0.5) * inputViewportInPixelsX * ARcpF1(outputSizeInPixelsX) - AF1_(0.5) + inputOffsetInPixelsX);
+    con0[3] = AU1_AF1(AF1_(0.5) * inputViewportInPixelsY * ARcpF1(outputSizeInPixelsY) - AF1_(0.5) + inputOffsetInPixelsY);
+}
+////////////////////////////////////////////////////////////////////////////////////////////////////////////////////////////////
+////////////////////////////////////////////////////////////////////////////////////////////////////////////////////////////////
+//_____________________________________________________________/\_______________________________________________________________
+//==============================================================================================================================
+//                                                   NON-PACKED 32-BIT VERSION
+//==============================================================================================================================
+#if defined(A_GPU)&&defined(FSR_EASU_F)
+ // Input callback prototypes, need to be implemented by calling shader
+ //AF4 FsrEasuRF(AF2 p);
+ //AF4 FsrEasuGF(AF2 p);
+ //AF4 FsrEasuBF(AF2 p);
+//------------------------------------------------------------------------------------------------------------------------------
+ // Filtering for a given tap for the scalar.
+ void FsrEasuTapF(
+ inout AF3 aC, // Accumulated color, with negative lobe.
+ inout AF1 aW, // Accumulated weight.
+ AF2 off, // Pixel offset from resolve position to tap.
+ AF2 dir, // Gradient direction.
+ AF2 len, // Length.
+ AF1 lob, // Negative lobe strength.
+ AF1 clp, // Clipping point.
+ AF3 c){ // Tap color.
+  // Rotate offset by direction.
+  AF2 v;
+  v.x=(off.x*( dir.x))+(off.y*dir.y);
+  v.y=(off.x*(-dir.y))+(off.y*dir.x);
+  // Anisotropy.
+  v*=len;
+  // Compute distance^2.
+  AF1 d2=v.x*v.x+v.y*v.y;
+  // Limit to the window as at corner, 2 taps can easily be outside.
+  d2=min(d2,clp);
+  // Approximation of lancos2 without sin() or rcp(), or sqrt() to get x.
+  //  (25/16 * (2/5 * x^2 - 1)^2 - (25/16 - 1)) * (1/4 * x^2 - 1)^2
+  //  |_______________________________________|   |_______________|
+  //                   base                             window
+  // The general form of the 'base' is,
+  //  (a*(b*x^2-1)^2-(a-1))
+  // Where 'a=1/(2*b-b^2)' and 'b' moves around the negative lobe.
+  AF1 wB=AF1_(2.0/5.0)*d2+AF1_(-1.0);
+  AF1 wA=lob*d2+AF1_(-1.0);
+  wB*=wB;
+  wA*=wA;
+  wB=AF1_(25.0/16.0)*wB+AF1_(-(25.0/16.0-1.0));
+  AF1 w=wB*wA;
+  // Do weighted average.
+  aC+=c*w;aW+=w;}
+//------------------------------------------------------------------------------------------------------------------------------
+ // Accumulate direction and length.
+ void FsrEasuSetF(
+ inout AF2 dir,
+ inout AF1 len,
+ AF2 pp,
+ AP1 biS,AP1 biT,AP1 biU,AP1 biV,
+ AF1 lA,AF1 lB,AF1 lC,AF1 lD,AF1 lE){
+  // Compute bilinear weight, branches factor out as predicates are compiler time immediates.
+  //  s t
+  //  u v
+  AF1 w = AF1_(0.0);
+  if(biS)w=(AF1_(1.0)-pp.x)*(AF1_(1.0)-pp.y);
+  if(biT)w=           pp.x *(AF1_(1.0)-pp.y);
+  if(biU)w=(AF1_(1.0)-pp.x)*           pp.y ;
+  if(biV)w=           pp.x *           pp.y ;
+  // Direction is the '+' diff.
+  //    a
+  //  b c d
+  //    e
+  // Then takes magnitude from abs average of both sides of 'c'.
+  // Length converts gradient reversal to 0, smoothly to non-reversal at 1, shaped, then adding horz and vert terms.
+  AF1 dc=lD-lC;
+  AF1 cb=lC-lB;
+  AF1 lenX=max(abs(dc),abs(cb));
+  lenX=APrxLoRcpF1(lenX);
+  AF1 dirX=lD-lB;
+  dir.x+=dirX*w;
+  lenX=ASatF1(abs(dirX)*lenX);
+  lenX*=lenX;
+  len+=lenX*w;
+  // Repeat for the y axis.
+  AF1 ec=lE-lC;
+  AF1 ca=lC-lA;
+  AF1 lenY=max(abs(ec),abs(ca));
+  lenY=APrxLoRcpF1(lenY);
+  AF1 dirY=lE-lA;
+  dir.y+=dirY*w;
+  lenY=ASatF1(abs(dirY)*lenY);
+  lenY*=lenY;
+  len+=lenY*w;}
+//------------------------------------------------------------------------------------------------------------------------------
+ void FsrEasuF(
+ out AF3 pix,
+ AU2 ip, // Integer pixel position in output.
+ AU4 con0, // Constants generated by FsrEasuCon().
+ AU4 con1,
+ AU4 con2,
+ AU4 con3){
+//------------------------------------------------------------------------------------------------------------------------------
+  // Get position of 'f'.
+  AF2 pp=AF2(ip)*AF2_AU2(con0.xy)+AF2_AU2(con0.zw);
+  AF2 fp=floor(pp);
+  pp-=fp;
+//------------------------------------------------------------------------------------------------------------------------------
+  // 12-tap kernel.
+  //    b c
+  //  e f g h
+  //  i j k l
+  //    n o
+  // Gather 4 ordering.
+  //  a b
+  //  r g
+  // For packed FP16, need either {rg} or {ab} so using the following setup for gather in all versions,
+  //    a b    <- unused (z)
+  //    r g
+  //  a b a b
+  //  r g r g
+  //    a b
+  //    r g    <- unused (z)
+  // Allowing dead-code removal to remove the 'z's.
+  AF2 p0=fp*AF2_AU2(con1.xy)+AF2_AU2(con1.zw);
+  // These are from p0 to avoid pulling two constants on pre-Navi hardware.
+  AF2 p1=p0+AF2_AU2(con2.xy);
+  AF2 p2=p0+AF2_AU2(con2.zw);
+  AF2 p3=p0+AF2_AU2(con3.xy);
+  AF4 bczzR=FsrEasuRF(p0);
+  AF4 bczzG=FsrEasuGF(p0);
+  AF4 bczzB=FsrEasuBF(p0);
+  AF4 ijfeR=FsrEasuRF(p1);
+  AF4 ijfeG=FsrEasuGF(p1);
+  AF4 ijfeB=FsrEasuBF(p1);
+  AF4 klhgR=FsrEasuRF(p2);
+  AF4 klhgG=FsrEasuGF(p2);
+  AF4 klhgB=FsrEasuBF(p2);
+  AF4 zzonR=FsrEasuRF(p3);
+  AF4 zzonG=FsrEasuGF(p3);
+  AF4 zzonB=FsrEasuBF(p3);
+//------------------------------------------------------------------------------------------------------------------------------
+  // Simplest multi-channel approximate luma possible (luma times 2, in 2 FMA/MAD).
+  AF4 bczzL=bczzB*AF4_(0.5)+(bczzR*AF4_(0.5)+bczzG);
+  AF4 ijfeL=ijfeB*AF4_(0.5)+(ijfeR*AF4_(0.5)+ijfeG);
+  AF4 klhgL=klhgB*AF4_(0.5)+(klhgR*AF4_(0.5)+klhgG);
+  AF4 zzonL=zzonB*AF4_(0.5)+(zzonR*AF4_(0.5)+zzonG);
+  // Rename.
+  AF1 bL=bczzL.x;
+  AF1 cL=bczzL.y;
+  AF1 iL=ijfeL.x;
+  AF1 jL=ijfeL.y;
+  AF1 fL=ijfeL.z;
+  AF1 eL=ijfeL.w;
+  AF1 kL=klhgL.x;
+  AF1 lL=klhgL.y;
+  AF1 hL=klhgL.z;
+  AF1 gL=klhgL.w;
+  AF1 oL=zzonL.z;
+  AF1 nL=zzonL.w;
+  // Accumulate for bilinear interpolation.
+  AF2 dir=AF2_(0.0);
+  AF1 len=AF1_(0.0);
+  FsrEasuSetF(dir,len,pp,true, false,false,false,bL,eL,fL,gL,jL);
+  FsrEasuSetF(dir,len,pp,false,true ,false,false,cL,fL,gL,hL,kL);
+  FsrEasuSetF(dir,len,pp,false,false,true ,false,fL,iL,jL,kL,nL);
+  FsrEasuSetF(dir,len,pp,false,false,false,true ,gL,jL,kL,lL,oL);
+//------------------------------------------------------------------------------------------------------------------------------
+  // Normalize with approximation, and cleanup close to zero.
+  AF2 dir2=dir*dir;
+  AF1 dirR=dir2.x+dir2.y;
+  AP1 zro=dirR<AF1_(1.0/32768.0);
+  dirR=APrxLoRsqF1(dirR);
+  dirR=zro?AF1_(1.0):dirR;
+  dir.x=zro?AF1_(1.0):dir.x;
+  dir*=AF2_(dirR);
+  // Transform from {0 to 2} to {0 to 1} range, and shape with square.
+  len=len*AF1_(0.5);
+  len*=len;
+  // Stretch kernel {1.0 vert|horz, to sqrt(2.0) on diagonal}.
+  AF1 stretch=(dir.x*dir.x+dir.y*dir.y)*APrxLoRcpF1(max(abs(dir.x),abs(dir.y)));
+  // Anisotropic length after rotation,
+  //  x := 1.0 lerp to 'stretch' on edges
+  //  y := 1.0 lerp to 2x on edges
+  AF2 len2=AF2(AF1_(1.0)+(stretch-AF1_(1.0))*len,AF1_(1.0)+AF1_(-0.5)*len);
+  // Based on the amount of 'edge',
+  // the window shifts from +/-{sqrt(2.0) to slightly beyond 2.0}.
+  AF1 lob=AF1_(0.5)+AF1_((1.0/4.0-0.04)-0.5)*len;
+  // Set distance^2 clipping point to the end of the adjustable window.
+  AF1 clp=APrxLoRcpF1(lob);
+//------------------------------------------------------------------------------------------------------------------------------
+  // Accumulation mixed with min/max of 4 nearest.
+  //    b c
+  //  e f g h
+  //  i j k l
+  //    n o
+  AF3 min4=min(AMin3F3(AF3(ijfeR.z,ijfeG.z,ijfeB.z),AF3(klhgR.w,klhgG.w,klhgB.w),AF3(ijfeR.y,ijfeG.y,ijfeB.y)),
+               AF3(klhgR.x,klhgG.x,klhgB.x));
+  AF3 max4=max(AMax3F3(AF3(ijfeR.z,ijfeG.z,ijfeB.z),AF3(klhgR.w,klhgG.w,klhgB.w),AF3(ijfeR.y,ijfeG.y,ijfeB.y)),
+               AF3(klhgR.x,klhgG.x,klhgB.x));
+  // Accumulation.
+  AF3 aC=AF3_(0.0);
+  AF1 aW=AF1_(0.0);
+  FsrEasuTapF(aC,aW,AF2( 0.0,-1.0)-pp,dir,len2,lob,clp,AF3(bczzR.x,bczzG.x,bczzB.x)); // b
+  FsrEasuTapF(aC,aW,AF2( 1.0,-1.0)-pp,dir,len2,lob,clp,AF3(bczzR.y,bczzG.y,bczzB.y)); // c
+  FsrEasuTapF(aC,aW,AF2(-1.0, 1.0)-pp,dir,len2,lob,clp,AF3(ijfeR.x,ijfeG.x,ijfeB.x)); // i
+  FsrEasuTapF(aC,aW,AF2( 0.0, 1.0)-pp,dir,len2,lob,clp,AF3(ijfeR.y,ijfeG.y,ijfeB.y)); // j
+  FsrEasuTapF(aC,aW,AF2( 0.0, 0.0)-pp,dir,len2,lob,clp,AF3(ijfeR.z,ijfeG.z,ijfeB.z)); // f
+  FsrEasuTapF(aC,aW,AF2(-1.0, 0.0)-pp,dir,len2,lob,clp,AF3(ijfeR.w,ijfeG.w,ijfeB.w)); // e
+  FsrEasuTapF(aC,aW,AF2( 1.0, 1.0)-pp,dir,len2,lob,clp,AF3(klhgR.x,klhgG.x,klhgB.x)); // k
+  FsrEasuTapF(aC,aW,AF2( 2.0, 1.0)-pp,dir,len2,lob,clp,AF3(klhgR.y,klhgG.y,klhgB.y)); // l
+  FsrEasuTapF(aC,aW,AF2( 2.0, 0.0)-pp,dir,len2,lob,clp,AF3(klhgR.z,klhgG.z,klhgB.z)); // h
+  FsrEasuTapF(aC,aW,AF2( 1.0, 0.0)-pp,dir,len2,lob,clp,AF3(klhgR.w,klhgG.w,klhgB.w)); // g
+  FsrEasuTapF(aC,aW,AF2( 1.0, 2.0)-pp,dir,len2,lob,clp,AF3(zzonR.z,zzonG.z,zzonB.z)); // o
+  FsrEasuTapF(aC,aW,AF2( 0.0, 2.0)-pp,dir,len2,lob,clp,AF3(zzonR.w,zzonG.w,zzonB.w)); // n
+//------------------------------------------------------------------------------------------------------------------------------
+  // Normalize and dering.
+  pix=min(max4,max(min4,aC*AF3_(ARcpF1(aW))));}
+#endif
+////////////////////////////////////////////////////////////////////////////////////////////////////////////////////////////////
+////////////////////////////////////////////////////////////////////////////////////////////////////////////////////////////////
+//_____________________________________________________________/\_______________________________________________________________
+//==============================================================================================================================
+//                                                    PACKED 16-BIT VERSION
+//==============================================================================================================================
+#if defined(A_GPU)&&defined(A_HALF)&&defined(FSR_EASU_H)
+// Input callback prototypes, need to be implemented by calling shader
+ //AH4 FsrEasuRH(AF2 p);
+ //AH4 FsrEasuGH(AF2 p);
+ //AH4 FsrEasuBH(AF2 p);
+//------------------------------------------------------------------------------------------------------------------------------
+ // This runs 2 taps in parallel.
+ void FsrEasuTapH(
+ inout AH2 aCR,inout AH2 aCG,inout AH2 aCB,
+ inout AH2 aW,
+ AH2 offX,AH2 offY,
+ AH2 dir,
+ AH2 len,
+ AH1 lob,
+ AH1 clp,
+ AH2 cR,AH2 cG,AH2 cB){
+  AH2 vX,vY;
+  vX=offX*  dir.xx +offY*dir.yy;
+  vY=offX*(-dir.yy)+offY*dir.xx;
+  vX*=len.x;vY*=len.y;
+  AH2 d2=vX*vX+vY*vY;
+  d2=min(d2,AH2_(clp));
+  AH2 wB=AH2_(2.0/5.0)*d2+AH2_(-1.0);
+  AH2 wA=AH2_(lob)*d2+AH2_(-1.0);
+  wB*=wB;
+  wA*=wA;
+  wB=AH2_(25.0/16.0)*wB+AH2_(-(25.0/16.0-1.0));
+  AH2 w=wB*wA;
+  aCR+=cR*w;aCG+=cG*w;aCB+=cB*w;aW+=w;}
+//------------------------------------------------------------------------------------------------------------------------------
+ // This runs 2 taps in parallel.
+ void FsrEasuSetH(
+ inout AH2 dirPX,inout AH2 dirPY,
+ inout AH2 lenP,
+ AH2 pp,
+ AP1 biST,AP1 biUV,
+ AH2 lA,AH2 lB,AH2 lC,AH2 lD,AH2 lE){
+  AH2 w = AH2_(0.0);
+  if(biST)w=(AH2(1.0,0.0)+AH2(-pp.x,pp.x))*AH2_(AH1_(1.0)-pp.y);
+  if(biUV)w=(AH2(1.0,0.0)+AH2(-pp.x,pp.x))*AH2_(          pp.y);
+  // ABS is not free in the packed FP16 path.
+  AH2 dc=lD-lC;
+  AH2 cb=lC-lB;
+  AH2 lenX=max(abs(dc),abs(cb));
+  lenX=ARcpH2(lenX);
+  AH2 dirX=lD-lB;
+  dirPX+=dirX*w;
+  lenX=ASatH2(abs(dirX)*lenX);
+  lenX*=lenX;
+  lenP+=lenX*w;
+  AH2 ec=lE-lC;
+  AH2 ca=lC-lA;
+  AH2 lenY=max(abs(ec),abs(ca));
+  lenY=ARcpH2(lenY);
+  AH2 dirY=lE-lA;
+  dirPY+=dirY*w;
+  lenY=ASatH2(abs(dirY)*lenY);
+  lenY*=lenY;
+  lenP+=lenY*w;}
+//------------------------------------------------------------------------------------------------------------------------------
+ void FsrEasuH(
+ out AH3 pix,
+ AU2 ip,
+ AU4 con0,
+ AU4 con1,
+ AU4 con2,
+ AU4 con3){
+//------------------------------------------------------------------------------------------------------------------------------
+  AF2 pp=AF2(ip)*AF2_AU2(con0.xy)+AF2_AU2(con0.zw);
+  AF2 fp=floor(pp);
+  pp-=fp;
+  AH2 ppp=AH2(pp);
+//------------------------------------------------------------------------------------------------------------------------------
+  AF2 p0=fp*AF2_AU2(con1.xy)+AF2_AU2(con1.zw);
+  AF2 p1=p0+AF2_AU2(con2.xy);
+  AF2 p2=p0+AF2_AU2(con2.zw);
+  AF2 p3=p0+AF2_AU2(con3.xy);
+  AH4 bczzR=FsrEasuRH(p0);
+  AH4 bczzG=FsrEasuGH(p0);
+  AH4 bczzB=FsrEasuBH(p0);
+  AH4 ijfeR=FsrEasuRH(p1);
+  AH4 ijfeG=FsrEasuGH(p1);
+  AH4 ijfeB=FsrEasuBH(p1);
+  AH4 klhgR=FsrEasuRH(p2);
+  AH4 klhgG=FsrEasuGH(p2);
+  AH4 klhgB=FsrEasuBH(p2);
+  AH4 zzonR=FsrEasuRH(p3);
+  AH4 zzonG=FsrEasuGH(p3);
+  AH4 zzonB=FsrEasuBH(p3);
+//------------------------------------------------------------------------------------------------------------------------------
+  AH4 bczzL=bczzB*AH4_(0.5)+(bczzR*AH4_(0.5)+bczzG);
+  AH4 ijfeL=ijfeB*AH4_(0.5)+(ijfeR*AH4_(0.5)+ijfeG);
+  AH4 klhgL=klhgB*AH4_(0.5)+(klhgR*AH4_(0.5)+klhgG);
+  AH4 zzonL=zzonB*AH4_(0.5)+(zzonR*AH4_(0.5)+zzonG);
+  AH1 bL=bczzL.x;
+  AH1 cL=bczzL.y;
+  AH1 iL=ijfeL.x;
+  AH1 jL=ijfeL.y;
+  AH1 fL=ijfeL.z;
+  AH1 eL=ijfeL.w;
+  AH1 kL=klhgL.x;
+  AH1 lL=klhgL.y;
+  AH1 hL=klhgL.z;
+  AH1 gL=klhgL.w;
+  AH1 oL=zzonL.z;
+  AH1 nL=zzonL.w;
+  // This part is different, accumulating 2 taps in parallel.
+  AH2 dirPX=AH2_(0.0);
+  AH2 dirPY=AH2_(0.0);
+  AH2 lenP=AH2_(0.0);
+  FsrEasuSetH(dirPX,dirPY,lenP,ppp,true, false,AH2(bL,cL),AH2(eL,fL),AH2(fL,gL),AH2(gL,hL),AH2(jL,kL));
+  FsrEasuSetH(dirPX,dirPY,lenP,ppp,false,true ,AH2(fL,gL),AH2(iL,jL),AH2(jL,kL),AH2(kL,lL),AH2(nL,oL));
+  AH2 dir=AH2(dirPX.r+dirPX.g,dirPY.r+dirPY.g);
+  AH1 len=lenP.r+lenP.g;
+//------------------------------------------------------------------------------------------------------------------------------
+  AH2 dir2=dir*dir;
+  AH1 dirR=dir2.x+dir2.y;
+  AP1 zro=dirR<AH1_(1.0/32768.0);
+  dirR=APrxLoRsqH1(dirR);
+  dirR=zro?AH1_(1.0):dirR;
+  dir.x=zro?AH1_(1.0):dir.x;
+  dir*=AH2_(dirR);
+  len=len*AH1_(0.5);
+  len*=len;
+  AH1 stretch=(dir.x*dir.x+dir.y*dir.y)*APrxLoRcpH1(max(abs(dir.x),abs(dir.y)));
+  AH2 len2=AH2(AH1_(1.0)+(stretch-AH1_(1.0))*len,AH1_(1.0)+AH1_(-0.5)*len);
+  AH1 lob=AH1_(0.5)+AH1_((1.0/4.0-0.04)-0.5)*len;
+  AH1 clp=APrxLoRcpH1(lob);
+//------------------------------------------------------------------------------------------------------------------------------
+  // FP16 is different, using packed trick to do min and max in same operation.
+  AH2 bothR=max(max(AH2(-ijfeR.z,ijfeR.z),AH2(-klhgR.w,klhgR.w)),max(AH2(-ijfeR.y,ijfeR.y),AH2(-klhgR.x,klhgR.x)));
+  AH2 bothG=max(max(AH2(-ijfeG.z,ijfeG.z),AH2(-klhgG.w,klhgG.w)),max(AH2(-ijfeG.y,ijfeG.y),AH2(-klhgG.x,klhgG.x)));
+  AH2 bothB=max(max(AH2(-ijfeB.z,ijfeB.z),AH2(-klhgB.w,klhgB.w)),max(AH2(-ijfeB.y,ijfeB.y),AH2(-klhgB.x,klhgB.x)));
+  // This part is different for FP16, working pairs of taps at a time.
+  AH2 pR=AH2_(0.0);
+  AH2 pG=AH2_(0.0);
+  AH2 pB=AH2_(0.0);
+  AH2 pW=AH2_(0.0);
+  FsrEasuTapH(pR,pG,pB,pW,AH2( 0.0, 1.0)-ppp.xx,AH2(-1.0,-1.0)-ppp.yy,dir,len2,lob,clp,bczzR.xy,bczzG.xy,bczzB.xy);
+  FsrEasuTapH(pR,pG,pB,pW,AH2(-1.0, 0.0)-ppp.xx,AH2( 1.0, 1.0)-ppp.yy,dir,len2,lob,clp,ijfeR.xy,ijfeG.xy,ijfeB.xy);
+  FsrEasuTapH(pR,pG,pB,pW,AH2( 0.0,-1.0)-ppp.xx,AH2( 0.0, 0.0)-ppp.yy,dir,len2,lob,clp,ijfeR.zw,ijfeG.zw,ijfeB.zw);
+  FsrEasuTapH(pR,pG,pB,pW,AH2( 1.0, 2.0)-ppp.xx,AH2( 1.0, 1.0)-ppp.yy,dir,len2,lob,clp,klhgR.xy,klhgG.xy,klhgB.xy);
+  FsrEasuTapH(pR,pG,pB,pW,AH2( 2.0, 1.0)-ppp.xx,AH2( 0.0, 0.0)-ppp.yy,dir,len2,lob,clp,klhgR.zw,klhgG.zw,klhgB.zw);
+  FsrEasuTapH(pR,pG,pB,pW,AH2( 1.0, 0.0)-ppp.xx,AH2( 2.0, 2.0)-ppp.yy,dir,len2,lob,clp,zzonR.zw,zzonG.zw,zzonB.zw);
+  AH3 aC=AH3(pR.x+pR.y,pG.x+pG.y,pB.x+pB.y);
+  AH1 aW=pW.x+pW.y;
+//------------------------------------------------------------------------------------------------------------------------------
+  // Slightly different for FP16 version due to combined min and max.
+  pix=min(AH3(bothR.y,bothG.y,bothB.y),max(-AH3(bothR.x,bothG.x,bothB.x),aC*AH3_(ARcpH1(aW))));}
+//==============================================================================================================================
+// An optimized AMD FSR's EASU implementation for Mobiles
+// Based on [REDACTED_PROFILE_URL]/FidelityFX-FSR/blob/master/ffx-fsr/ffx_fsr1.h
+// Details can be found: https://atyuwen.github.io/posts/optimizing-fsr/
+// 
+// Distributed under the MIT License. Copyright (c) 2021 Yuwen Wu
+// Permission is hereby granted, free of charge, to any person obtaining a copy
+// of this softwareand associated documentation files(the "Software"), to deal
+// in the Software without restriction, including without limitation the rights
+// to use, copy, modify, merge, publish, distribute, sublicense, and /or sell
+// copies of the Software, and to permit persons to whom the Software is
+// furnished to do so, subject to the following conditions :
+// The above copyright noticeand this permission notice shall be included in all
+// copies or substantial portions of the Software.
+// THE SOFTWARE IS PROVIDED "AS IS", WITHOUT WARRANTY OF ANY KIND, EXPRESS OR
+// IMPLIED, INCLUDING BUT NOT LIMITED TO THE WARRANTIES OF MERCHANTABILITY,
+// FITNESS FOR A PARTICULAR PURPOSE AND NONINFRINGEMENT.IN NO EVENT SHALL THE
+// AUTHORS OR COPYRIGHT HOLDERS BE LIABLE FOR ANY CLAIM, DAMAGES OR OTHER
+// LIABILITY, WHETHER IN AN ACTION OF CONTRACT, TORT OR OTHERWISE, ARISING FROM,
+// OUT OF OR IN CONNECTION WITH THE SOFTWARE OR THE USE OR OTHER DEALINGS IN THE
+// SOFTWARE.
+//
+// Qualcomm Adreno modifications
+// Copyright (c) 2022, Qualcomm Innovation Center, Inc. All rights reserved.
+// 
+// -- FsrEasuSampleH should be implemented by calling shader, like following:
+//    AH3 FsrEasuSampleH(AF2 p) { return MyTex.SampleLevel(LinearSampler, p, 0).xyz; }
+//==============================================================================================================================
+ void FsrEasuL(
+ out AH3 pix,
+ AU2 ip,
+ AU4 con0,
+ AU4 con1,
+ AU4 con2,
+ AU4 con3){
+//------------------------------------------------------------------------------------------------------------------------------
+  // Direction is the '+' diff.
+  //    A
+  //  B C D
+  //    E
+  AF2 pp=AF2(ip)*AF2_AU2(con0.xy)+AF2_AU2(con0.zw);
+  AF2 tc=(pp+AF2_(0.5))*AF2_AU2(con1.xy);
+  AH3 sA=FsrEasuSampleH(tc-AF2(0, AF1_AU1(con1.y)));
+  AH3 sB=FsrEasuSampleH(tc-AF2(AF1_AU1(con1.x), 0));
+  AH3 sC=FsrEasuSampleH(tc);
+  AH3 sD=FsrEasuSampleH(tc+AF2(AF1_AU1(con1.x), 0));
+  AH3 sE=FsrEasuSampleH(tc+AF2(0, AF1_AU1(con1.y)));
+  AH1 lA=sA.r*AH1_(0.5)+sA.g;
+  AH1 lB=sB.r*AH1_(0.5)+sB.g;
+  AH1 lC=sC.r*AH1_(0.5)+sC.g;
+  AH1 lD=sD.r*AH1_(0.5)+sD.g;
+  AH1 lE=sE.r*AH1_(0.5)+sE.g;
+  // Then takes magnitude from abs average of both sides of 'C'.
+  // Length converts gradient reversal to 0, smoothly to non-reversal at 1, shaped, then adding horz and vert terms.
+  AH1 dc=lD-lC;
+  AH1 cb=lC-lB;
+  AH1 lenX=max(abs(dc),abs(cb));
+  lenX=ARcpH1(lenX);
+  AH1 dirX=lD-lB;
+  lenX=ASatH1(abs(dirX)*lenX);
+  lenX*=lenX;
+  // Repeat for the y axis.
+  AH1 ec=lE-lC;
+  AH1 ca=lC-lA;
+  AH1 lenY=max(abs(ec),abs(ca));
+  lenY=ARcpH1(lenY);
+  AH1 dirY=lE-lA;
+  lenY=ASatH1(abs(dirY)*lenY);
+  AH1 len = lenY * lenY + lenX;
+  AH2 dir = AH2(dirX, dirY);
+ //------------------------------------------------------------------------------------------------------------------------------ 
+  AH2 dir2=dir*dir;
+  AH1 dirR=dir2.x+dir2.y;
+  if (dirR<AH1_(1.0/64.0)) {
+    pix = sC;
+    return;
+  }
+  dirR=ARsqH1(dirR);
+  dir*=AH2_(dirR);
+  len=len*AH1_(0.5);
+  len*=len;
+  AH1 stretch=(dir.x*dir.x+dir.y*dir.y)*ARcpH1(max(abs(dir.x),abs(dir.y)));
+  AH2 len2=AH2(AH1_(1.0)+(stretch-AH1_(1.0))*len,AH1_(1.0)+AH1_(-0.5)*len);
+  AH1 lob=AH1_(0.5)+AH1_((1.0/4.0-0.04)-0.5)*len;
+  AH1 clp=ARcpH1(lob);
+//------------------------------------------------------------------------------------------------------------------------------
+  AF2 fp=floor(pp);
+  pp-=fp;
+  AH2 ppp=AH2(pp);
+  AF2 p0=fp*AF2_AU2(con1.xy)+AF2_AU2(con1.zw);
+  AF2 p1=p0+AF2_AU2(con2.xy);
+  AF2 p2=p0+AF2_AU2(con2.zw);
+  AF2 p3=p0+AF2_AU2(con3.xy);
+  p0.y-=AF1_AU1(con1.w);
+
+  AH4 fgcbR=FsrEasuRH(p0);
+  AH4 fgcbG=FsrEasuGH(p0);
+  AH4 fgcbB=FsrEasuBH(p0);
+  AH4 ijfeR=FsrEasuRH(p1);
+  AH4 ijfeG=FsrEasuGH(p1);
+  AH4 ijfeB=FsrEasuBH(p1);
+  AH4 klhgR=FsrEasuRH(p2);
+  AH4 klhgG=FsrEasuGH(p2);
+  AH4 klhgB=FsrEasuBH(p2);
+  AH4 nokjR=FsrEasuRH(p3);
+  AH4 nokjG=FsrEasuGH(p3);
+  AH4 nokjB=FsrEasuBH(p3);
+
+//------------------------------------------------------------------------------------------------------------------------------
+  // This part is different for FP16, working pairs of taps at a time.
+  AH2 pR=AH2_(0.0);
+  AH2 pG=AH2_(0.0);
+  AH2 pB=AH2_(0.0);
+  AH2 pW=AH2_(0.0);
+  FsrEasuTapH(pR,pG,pB,pW,AH2( 1.0, 0.0)-ppp.xx,AH2(-1.0,-1.0)-ppp.yy,dir,len2,lob,clp,fgcbR.zw,fgcbG.zw,fgcbB.zw);
+  FsrEasuTapH(pR,pG,pB,pW,AH2(-1.0, 0.0)-ppp.xx,AH2( 1.0, 1.0)-ppp.yy,dir,len2,lob,clp,ijfeR.xy,ijfeG.xy,ijfeB.xy);
+  FsrEasuTapH(pR,pG,pB,pW,AH2( 0.0,-1.0)-ppp.xx,AH2( 0.0, 0.0)-ppp.yy,dir,len2,lob,clp,ijfeR.zw,ijfeG.zw,ijfeB.zw);
+  FsrEasuTapH(pR,pG,pB,pW,AH2( 1.0, 2.0)-ppp.xx,AH2( 1.0, 1.0)-ppp.yy,dir,len2,lob,clp,klhgR.xy,klhgG.xy,klhgB.xy);
+  FsrEasuTapH(pR,pG,pB,pW,AH2( 2.0, 1.0)-ppp.xx,AH2( 0.0, 0.0)-ppp.yy,dir,len2,lob,clp,klhgR.zw,klhgG.zw,klhgB.zw);
+  FsrEasuTapH(pR,pG,pB,pW,AH2( 0.0, 1.0)-ppp.xx,AH2( 2.0, 2.0)-ppp.yy,dir,len2,lob,clp,nokjR.xy,nokjG.xy,nokjB.xy);
+  AH3 aC=AH3(pR.x+pR.y,pG.x+pG.y,pB.x+pB.y);
+  AH1 aW=pW.x+pW.y;
+//------------------------------------------------------------------------------------------------------------------------------
+  // Qualcomm - normalize and de-ring.  While the Atyuwen original blog post did not see any visual changes from removing the de-ringer
+  //            we saw noticable artifacts on high contrast sharp edges.
+  AH3 min4 = min(AMin3H3(AH3(ijfeR.z, ijfeG.z, ijfeB.z), AH3(klhgR.w, klhgG.w, klhgB.w), AH3(ijfeR.y, ijfeG.y, ijfeB.y)),
+      AH3(klhgR.x, klhgG.x, klhgB.x));
+  AH3 max4 = max(AMax3H3(AH3(ijfeR.z, ijfeG.z, ijfeB.z), AH3(klhgR.w, klhgG.w, klhgB.w), AH3(ijfeR.y, ijfeG.y, ijfeB.y)),
+      AH3(klhgR.x, klhgG.x, klhgB.x));
+  pix = min(max4, max(min4, aC * AH3_(ARcpH1(aW))));}
+#endif
+////////////////////////////////////////////////////////////////////////////////////////////////////////////////////////////////
+////////////////////////////////////////////////////////////////////////////////////////////////////////////////////////////////
+////////////////////////////////////////////////////////////////////////////////////////////////////////////////////////////////
+////////////////////////////////////////////////////////////////////////////////////////////////////////////////////////////////
+//_____________________________________________________________/\_______________________________________________________________
+//==============================================================================================================================
+//
+//                                      FSR - [RCAS] ROBUST CONTRAST ADAPTIVE SHARPENING
+//
+//------------------------------------------------------------------------------------------------------------------------------
+// CAS uses a simplified mechanism to convert local contrast into a variable amount of sharpness.
+// RCAS uses a more exact mechanism, solving for the maximum local sharpness possible before clipping.
+// RCAS also has a built in process to limit sharpening of what it detects as possible noise.
+// RCAS sharper does not support scaling, as it should be applied after EASU scaling.
+// Pass EASU output straight into RCAS, no color conversions necessary.
+//------------------------------------------------------------------------------------------------------------------------------
+// RCAS is based on the following logic.
+// RCAS uses a 5 tap filter in a cross pattern (same as CAS),
+//    w                n
+//  w 1 w  for taps  w m e 
+//    w                s
+// Where 'w' is the negative lobe weight.
+//  output = (w*(n+e+w+s)+m)/(4*w+1)
+// RCAS solves for 'w' by seeing where the signal might clip out of the {0 to 1} input range,
+//  0 == (w*(n+e+w+s)+m)/(4*w+1) -> w = -m/(n+e+w+s)
+//  1 == (w*(n+e+w+s)+m)/(4*w+1) -> w = (1-m)/(n+e+w+s-4*1)
+// Then chooses the 'w' which results in no clipping, limits 'w', and multiplies by the 'sharp' amount.
+// This solution above has issues with MSAA input as the steps along the gradient cause edge detection issues.
+// So RCAS uses 4x the maximum and 4x the minimum (depending on equation)in place of the individual taps.
+// As well as switching from 'm' to either the minimum or maximum (depending on side), to help in energy conservation.
+// This stabilizes RCAS.
+// RCAS does a simple highpass which is normalized against the local contrast then shaped,
+//       0.25
+//  0.25  -1  0.25
+//       0.25
+// This is used as a noise detection filter, to reduce the effect of RCAS on grain, and focus on real edges.
+//
+//  GLSL example for the required callbacks :
+// 
+//  AH4 FsrRcasLoadH(ASW2 p){return AH4(imageLoad(imgSrc,ASU2(p)));}
+//  void FsrRcasInputH(inout AH1 r,inout AH1 g,inout AH1 b)
+//  {
+//    //do any simple input color conversions here or leave empty if none needed
+//  }
+//  
+//  FsrRcasCon need to be called from the CPU or GPU to set up constants.
+//  Including a GPU example here, the 'con' value would be stored out to a constant buffer.
+// 
+//  AU4 con;
+//  FsrRcasCon(con,
+//   0.0); // The scale is {0.0 := maximum sharpness, to N>0, where N is the number of stops (halving) of the reduction of sharpness}.
+// ---------------
+// RCAS sharpening supports a CAS-like pass-through alpha via,
+//  #define FSR_RCAS_PASSTHROUGH_ALPHA 1
+// RCAS also supports a define to enable a more expensive path to avoid some sharpening of noise.
+// Would suggest it is better to apply film grain after RCAS sharpening (and after scaling) instead of using this define,
+//  #define FSR_RCAS_DENOISE 1
+//==============================================================================================================================
+// This is set at the limit of providing unnatural results for sharpening.
+#define FSR_RCAS_LIMIT (0.25-(1.0/16.0))
+////////////////////////////////////////////////////////////////////////////////////////////////////////////////////////////////
+////////////////////////////////////////////////////////////////////////////////////////////////////////////////////////////////
+//_____________________________________________________________/\_______________________________________________________________
+//==============================================================================================================================
+//                                                      CONSTANT SETUP
+//==============================================================================================================================
+// Call to setup required constant values (works on CPU or GPU).
+A_STATIC void FsrRcasCon(
+outAU4 con,
+// The scale is {0.0 := maximum, to N>0, where N is the number of stops (halving) of the reduction of sharpness}.
+AF1 sharpness){
+ // Transform from stops to linear value.
+ sharpness=AExp2F1(-sharpness);
+ varAF2(hSharp)=initAF2(sharpness,sharpness);
+ con[0]=AU1_AF1(sharpness);
+ con[1]=AU1_AH2_AF2(hSharp);
+ con[2]=0;
+ con[3]=0;}
+////////////////////////////////////////////////////////////////////////////////////////////////////////////////////////////////
+////////////////////////////////////////////////////////////////////////////////////////////////////////////////////////////////
+//_____________________________________________________________/\_______________________________________________________________
+//==============================================================================================================================
+//                                                   NON-PACKED 32-BIT VERSION
+//==============================================================================================================================
+#if defined(A_GPU)&&defined(FSR_RCAS_F)
+ // Input callback prototypes that need to be implemented by calling shader
+ //AF4 FsrRcasLoadF(ASU2 p);
+ //void FsrRcasInputF(inout AF1 r,inout AF1 g,inout AF1 b);
+//------------------------------------------------------------------------------------------------------------------------------
+ void FsrRcasF(
+ out AF1 pixR, // Output values, non-vector so port between RcasFilter() and RcasFilterH() is easy.
+ out AF1 pixG,
+ out AF1 pixB,
+ #ifdef FSR_RCAS_PASSTHROUGH_ALPHA
+  out AF1 pixA,
+ #endif
+ AU2 ip, // Integer pixel position in output.
+ AU4 con){ // Constant generated by RcasSetup().
+  // Algorithm uses minimal 3x3 pixel neighborhood.
+  //    b 
+  //  d e f
+  //    h
+  ASU2 sp=ASU2(ip);
+  AF3 b=FsrRcasLoadF(sp+ASU2( 0,-1)).rgb;
+  AF3 d=FsrRcasLoadF(sp+ASU2(-1, 0)).rgb;
+  #ifdef FSR_RCAS_PASSTHROUGH_ALPHA
+   AF4 ee=FsrRcasLoadF(sp);
+   AF3 e=ee.rgb;pixA=ee.a;
+  #else
+   AF3 e=FsrRcasLoadF(sp).rgb;
+  #endif
+  AF3 f=FsrRcasLoadF(sp+ASU2( 1, 0)).rgb;
+  AF3 h=FsrRcasLoadF(sp+ASU2( 0, 1)).rgb;
+  // Rename (32-bit) or regroup (16-bit).
+  AF1 bR=b.r;
+  AF1 bG=b.g;
+  AF1 bB=b.b;
+  AF1 dR=d.r;
+  AF1 dG=d.g;
+  AF1 dB=d.b;
+  AF1 eR=e.r;
+  AF1 eG=e.g;
+  AF1 eB=e.b;
+  AF1 fR=f.r;
+  AF1 fG=f.g;
+  AF1 fB=f.b;
+  AF1 hR=h.r;
+  AF1 hG=h.g;
+  AF1 hB=h.b;
+  // Run optional input transform.
+  FsrRcasInputF(bR,bG,bB);
+  FsrRcasInputF(dR,dG,dB);
+  FsrRcasInputF(eR,eG,eB);
+  FsrRcasInputF(fR,fG,fB);
+  FsrRcasInputF(hR,hG,hB);
+  // Luma times 2.
+  AF1 bL=bB*AF1_(0.5)+(bR*AF1_(0.5)+bG);
+  AF1 dL=dB*AF1_(0.5)+(dR*AF1_(0.5)+dG);
+  AF1 eL=eB*AF1_(0.5)+(eR*AF1_(0.5)+eG);
+  AF1 fL=fB*AF1_(0.5)+(fR*AF1_(0.5)+fG);
+  AF1 hL=hB*AF1_(0.5)+(hR*AF1_(0.5)+hG);
+  // Noise detection.
+  AF1 nz=AF1_(0.25)*bL+AF1_(0.25)*dL+AF1_(0.25)*fL+AF1_(0.25)*hL-eL;
+  nz=ASatF1(abs(nz)*APrxMedRcpF1(AMax3F1(AMax3F1(bL,dL,eL),fL,hL)-AMin3F1(AMin3F1(bL,dL,eL),fL,hL)));
+  nz=AF1_(-0.5)*nz+AF1_(1.0);
+  // Min and max of ring.
+  AF1 mn4R=min(AMin3F1(bR,dR,fR),hR);
+  AF1 mn4G=min(AMin3F1(bG,dG,fG),hG);
+  AF1 mn4B=min(AMin3F1(bB,dB,fB),hB);
+  AF1 mx4R=max(AMax3F1(bR,dR,fR),hR);
+  AF1 mx4G=max(AMax3F1(bG,dG,fG),hG);
+  AF1 mx4B=max(AMax3F1(bB,dB,fB),hB);
+  // Immediate constants for peak range.
+  AF2 peakC=AF2(1.0,-1.0*4.0);
+  // Limiters, these need to be high precision RCPs.
+  AF1 hitMinR=min(mn4R,eR)*ARcpF1(AF1_(4.0)*mx4R);
+  AF1 hitMinG=min(mn4G,eG)*ARcpF1(AF1_(4.0)*mx4G);
+  AF1 hitMinB=min(mn4B,eB)*ARcpF1(AF1_(4.0)*mx4B);
+  AF1 hitMaxR=(peakC.x-max(mx4R,eR))*ARcpF1(AF1_(4.0)*mn4R+peakC.y);
+  AF1 hitMaxG=(peakC.x-max(mx4G,eG))*ARcpF1(AF1_(4.0)*mn4G+peakC.y);
+  AF1 hitMaxB=(peakC.x-max(mx4B,eB))*ARcpF1(AF1_(4.0)*mn4B+peakC.y);
+  AF1 lobeR=max(-hitMinR,hitMaxR);
+  AF1 lobeG=max(-hitMinG,hitMaxG);
+  AF1 lobeB=max(-hitMinB,hitMaxB);
+  AF1 lobe=max(AF1_(-FSR_RCAS_LIMIT),min(AMax3F1(lobeR,lobeG,lobeB),AF1_(0.0)))*AF1_AU1(con.x);
+  // Apply noise removal.
+  #ifdef FSR_RCAS_DENOISE
+   lobe*=nz;
+  #endif
+  // Resolve, which needs the medium precision rcp approximation to avoid visible tonality changes.
+  AF1 rcpL=APrxMedRcpF1(AF1_(4.0)*lobe+AF1_(1.0));
+  pixR=(lobe*bR+lobe*dR+lobe*hR+lobe*fR+eR)*rcpL;
+  pixG=(lobe*bG+lobe*dG+lobe*hG+lobe*fG+eG)*rcpL;
+  pixB=(lobe*bB+lobe*dB+lobe*hB+lobe*fB+eB)*rcpL;
+  return;} 
+#endif
+////////////////////////////////////////////////////////////////////////////////////////////////////////////////////////////////
+////////////////////////////////////////////////////////////////////////////////////////////////////////////////////////////////
+//_____________________________________________________________/\_______________________________________________________________
+//==============================================================================================================================
+//                                                  NON-PACKED 16-BIT VERSION
+//==============================================================================================================================
+#if defined(A_GPU)&&defined(A_HALF)&&defined(FSR_RCAS_H)
+ // Input callback prototypes that need to be implemented by calling shader
+ //AH4 FsrRcasLoadH(ASW2 p);
+ //void FsrRcasInputH(inout AH1 r,inout AH1 g,inout AH1 b);
+//------------------------------------------------------------------------------------------------------------------------------
+ void FsrRcasH(
+ out AH1 pixR, // Output values, non-vector so port between RcasFilter() and RcasFilterH() is easy.
+ out AH1 pixG,
+ out AH1 pixB,
+ #ifdef FSR_RCAS_PASSTHROUGH_ALPHA
+  out AH1 pixA,
+ #endif
+ AU2 ip, // Integer pixel position in output.
+ AU4 con){ // Constant generated by RcasSetup().
+  // Sharpening algorithm uses minimal 3x3 pixel neighborhood.
+  //    b 
+  //  d e f
+  //    h
+  ASW2 sp=ASW2(ip);
+  AH3 b=FsrRcasLoadH(sp+ASW2( 0,-1)).rgb;
+  AH3 d=FsrRcasLoadH(sp+ASW2(-1, 0)).rgb;
+  #ifdef FSR_RCAS_PASSTHROUGH_ALPHA
+   AH4 ee=FsrRcasLoadH(sp);
+   AH3 e=ee.rgb;pixA=ee.a;
+  #else
+   AH3 e=FsrRcasLoadH(sp).rgb;
+  #endif
+  AH3 f=FsrRcasLoadH(sp+ASW2( 1, 0)).rgb;
+  AH3 h=FsrRcasLoadH(sp+ASW2( 0, 1)).rgb;
+  // Rename (32-bit) or regroup (16-bit).
+  AH1 bR=b.r;
+  AH1 bG=b.g;
+  AH1 bB=b.b;
+  AH1 dR=d.r;
+  AH1 dG=d.g;
+  AH1 dB=d.b;
+  AH1 eR=e.r;
+  AH1 eG=e.g;
+  AH1 eB=e.b;
+  AH1 fR=f.r;
+  AH1 fG=f.g;
+  AH1 fB=f.b;
+  AH1 hR=h.r;
+  AH1 hG=h.g;
+  AH1 hB=h.b;
+  // Run optional input transform.
+  FsrRcasInputH(bR,bG,bB);
+  FsrRcasInputH(dR,dG,dB);
+  FsrRcasInputH(eR,eG,eB);
+  FsrRcasInputH(fR,fG,fB);
+  FsrRcasInputH(hR,hG,hB);
+  // Luma times 2.
+  AH1 bL=bB*AH1_(0.5)+(bR*AH1_(0.5)+bG);
+  AH1 dL=dB*AH1_(0.5)+(dR*AH1_(0.5)+dG);
+  AH1 eL=eB*AH1_(0.5)+(eR*AH1_(0.5)+eG);
+  AH1 fL=fB*AH1_(0.5)+(fR*AH1_(0.5)+fG);
+  AH1 hL=hB*AH1_(0.5)+(hR*AH1_(0.5)+hG);
+  // Noise detection.
+  AH1 nz=AH1_(0.25)*bL+AH1_(0.25)*dL+AH1_(0.25)*fL+AH1_(0.25)*hL-eL;
+  nz=ASatH1(abs(nz)*APrxMedRcpH1(AMax3H1(AMax3H1(bL,dL,eL),fL,hL)-AMin3H1(AMin3H1(bL,dL,eL),fL,hL)));
+  nz=AH1_(-0.5)*nz+AH1_(1.0);
+  // Min and max of ring.
+  AH1 mn4R=min(AMin3H1(bR,dR,fR),hR);
+  AH1 mn4G=min(AMin3H1(bG,dG,fG),hG);
+  AH1 mn4B=min(AMin3H1(bB,dB,fB),hB);
+  AH1 mx4R=max(AMax3H1(bR,dR,fR),hR);
+  AH1 mx4G=max(AMax3H1(bG,dG,fG),hG);
+  AH1 mx4B=max(AMax3H1(bB,dB,fB),hB);
+  // Immediate constants for peak range.
+  AH2 peakC=AH2(1.0,-1.0*4.0);
+  // Limiters, these need to be high precision RCPs.
+  AH1 hitMinR=min(mn4R,eR)*ARcpH1(AH1_(4.0)*mx4R);
+  AH1 hitMinG=min(mn4G,eG)*ARcpH1(AH1_(4.0)*mx4G);
+  AH1 hitMinB=min(mn4B,eB)*ARcpH1(AH1_(4.0)*mx4B);
+  AH1 hitMaxR=(peakC.x-max(mx4R,eR))*ARcpH1(AH1_(4.0)*mn4R+peakC.y);
+  AH1 hitMaxG=(peakC.x-max(mx4G,eG))*ARcpH1(AH1_(4.0)*mn4G+peakC.y);
+  AH1 hitMaxB=(peakC.x-max(mx4B,eB))*ARcpH1(AH1_(4.0)*mn4B+peakC.y);
+  AH1 lobeR=max(-hitMinR,hitMaxR);
+  AH1 lobeG=max(-hitMinG,hitMaxG);
+  AH1 lobeB=max(-hitMinB,hitMaxB);
+  AH1 lobe=max(AH1_(-FSR_RCAS_LIMIT),min(AMax3H1(lobeR,lobeG,lobeB),AH1_(0.0)))*AH2_AU1(con.y).x;
+  // Apply noise removal.
+  #ifdef FSR_RCAS_DENOISE
+   lobe*=nz;
+  #endif
+  // Resolve, which needs the medium precision rcp approximation to avoid visible tonality changes.
+  AH1 rcpL=APrxMedRcpH1(AH1_(4.0)*lobe+AH1_(1.0));
+  pixR=(lobe*bR+lobe*dR+lobe*hR+lobe*fR+eR)*rcpL;
+  pixG=(lobe*bG+lobe*dG+lobe*hG+lobe*fG+eG)*rcpL;
+  pixB=(lobe*bB+lobe*dB+lobe*hB+lobe*fB+eB)*rcpL;}
+#endif
+////////////////////////////////////////////////////////////////////////////////////////////////////////////////////////////////
+////////////////////////////////////////////////////////////////////////////////////////////////////////////////////////////////
+//_____________________________________________________________/\_______________________________________________________________
+//==============================================================================================================================
+//                                                     PACKED 16-BIT VERSION
+//==============================================================================================================================
+#if defined(A_GPU)&&defined(A_HALF)&&defined(FSR_RCAS_HX2)
+ // Input callback prototypes that need to be implemented by the calling shader
+ //AH4 FsrRcasLoadHx2(ASW2 p);
+ //void FsrRcasInputHx2(inout AH2 r,inout AH2 g,inout AH2 b);
+//------------------------------------------------------------------------------------------------------------------------------
+ // Can be used to convert from packed Structures of Arrays to Arrays of Structures for store.
+ void FsrRcasDepackHx2(out AH4 pix0,out AH4 pix1,AH2 pixR,AH2 pixG,AH2 pixB){
+  #ifdef A_HLSL
+   // Invoke a slower path for DX only, since it won't allow uninitialized values.
+   pix0.a=pix1.a=0.0;
+  #endif
+  pix0.rgb=AH3(pixR.x,pixG.x,pixB.x);
+  pix1.rgb=AH3(pixR.y,pixG.y,pixB.y);}
+//------------------------------------------------------------------------------------------------------------------------------
+ void FsrRcasHx2(
+ // Output values are for 2 8x8 tiles in a 16x8 region.
+ //  pix<R,G,B>.x =  left 8x8 tile
+ //  pix<R,G,B>.y = right 8x8 tile
+ // This enables later processing to easily be packed as well.
+ out AH2 pixR,
+ out AH2 pixG,
+ out AH2 pixB,
+ #ifdef FSR_RCAS_PASSTHROUGH_ALPHA
+  out AH2 pixA,
+ #endif
+ AU2 ip, // Integer pixel position in output.
+ AU4 con){ // Constant generated by RcasSetup().
+  // No scaling algorithm uses minimal 3x3 pixel neighborhood.
+  ASW2 sp0=ASW2(ip);
+  AH3 b0=FsrRcasLoadHx2(sp0+ASW2( 0,-1)).rgb;
+  AH3 d0=FsrRcasLoadHx2(sp0+ASW2(-1, 0)).rgb;
+  #ifdef FSR_RCAS_PASSTHROUGH_ALPHA
+   AH4 ee0=FsrRcasLoadHx2(sp0);
+   AH3 e0=ee0.rgb;pixA.r=ee0.a;
+  #else
+   AH3 e0=FsrRcasLoadHx2(sp0).rgb;
+  #endif
+  AH3 f0=FsrRcasLoadHx2(sp0+ASW2( 1, 0)).rgb;
+  AH3 h0=FsrRcasLoadHx2(sp0+ASW2( 0, 1)).rgb;
+  ASW2 sp1=sp0+ASW2(8,0);
+  AH3 b1=FsrRcasLoadHx2(sp1+ASW2( 0,-1)).rgb;
+  AH3 d1=FsrRcasLoadHx2(sp1+ASW2(-1, 0)).rgb;
+  #ifdef FSR_RCAS_PASSTHROUGH_ALPHA
+   AH4 ee1=FsrRcasLoadHx2(sp1);
+   AH3 e1=ee1.rgb;pixA.g=ee1.a;
+  #else
+   AH3 e1=FsrRcasLoadHx2(sp1).rgb;
+  #endif
+  AH3 f1=FsrRcasLoadHx2(sp1+ASW2( 1, 0)).rgb;
+  AH3 h1=FsrRcasLoadHx2(sp1+ASW2( 0, 1)).rgb;
+  // Arrays of Structures to Structures of Arrays conversion.
+  AH2 bR=AH2(b0.r,b1.r);
+  AH2 bG=AH2(b0.g,b1.g);
+  AH2 bB=AH2(b0.b,b1.b);
+  AH2 dR=AH2(d0.r,d1.r);
+  AH2 dG=AH2(d0.g,d1.g);
+  AH2 dB=AH2(d0.b,d1.b);
+  AH2 eR=AH2(e0.r,e1.r);
+  AH2 eG=AH2(e0.g,e1.g);
+  AH2 eB=AH2(e0.b,e1.b);
+  AH2 fR=AH2(f0.r,f1.r);
+  AH2 fG=AH2(f0.g,f1.g);
+  AH2 fB=AH2(f0.b,f1.b);
+  AH2 hR=AH2(h0.r,h1.r);
+  AH2 hG=AH2(h0.g,h1.g);
+  AH2 hB=AH2(h0.b,h1.b);
+  // Run optional input transform.
+  FsrRcasInputHx2(bR,bG,bB);
+  FsrRcasInputHx2(dR,dG,dB);
+  FsrRcasInputHx2(eR,eG,eB);
+  FsrRcasInputHx2(fR,fG,fB);
+  FsrRcasInputHx2(hR,hG,hB);
+  // Luma times 2.
+  AH2 bL=bB*AH2_(0.5)+(bR*AH2_(0.5)+bG);
+  AH2 dL=dB*AH2_(0.5)+(dR*AH2_(0.5)+dG);
+  AH2 eL=eB*AH2_(0.5)+(eR*AH2_(0.5)+eG);
+  AH2 fL=fB*AH2_(0.5)+(fR*AH2_(0.5)+fG);
+  AH2 hL=hB*AH2_(0.5)+(hR*AH2_(0.5)+hG);
+  // Noise detection.
+  AH2 nz=AH2_(0.25)*bL+AH2_(0.25)*dL+AH2_(0.25)*fL+AH2_(0.25)*hL-eL;
+  nz=ASatH2(abs(nz)*APrxMedRcpH2(AMax3H2(AMax3H2(bL,dL,eL),fL,hL)-AMin3H2(AMin3H2(bL,dL,eL),fL,hL)));
+  nz=AH2_(-0.5)*nz+AH2_(1.0);
+  // Min and max of ring.
+  AH2 mn4R=min(AMin3H2(bR,dR,fR),hR);
+  AH2 mn4G=min(AMin3H2(bG,dG,fG),hG);
+  AH2 mn4B=min(AMin3H2(bB,dB,fB),hB);
+  AH2 mx4R=max(AMax3H2(bR,dR,fR),hR);
+  AH2 mx4G=max(AMax3H2(bG,dG,fG),hG);
+  AH2 mx4B=max(AMax3H2(bB,dB,fB),hB);
+  // Immediate constants for peak range.
+  AH2 peakC=AH2(1.0,-1.0*4.0);
+  // Limiters, these need to be high precision RCPs.
+  AH2 hitMinR=min(mn4R,eR)*ARcpH2(AH2_(4.0)*mx4R);
+  AH2 hitMinG=min(mn4G,eG)*ARcpH2(AH2_(4.0)*mx4G);
+  AH2 hitMinB=min(mn4B,eB)*ARcpH2(AH2_(4.0)*mx4B);
+  AH2 hitMaxR=(peakC.x-max(mx4R,eR))*ARcpH2(AH2_(4.0)*mn4R+peakC.y);
+  AH2 hitMaxG=(peakC.x-max(mx4G,eG))*ARcpH2(AH2_(4.0)*mn4G+peakC.y);
+  AH2 hitMaxB=(peakC.x-max(mx4B,eB))*ARcpH2(AH2_(4.0)*mn4B+peakC.y);
+  AH2 lobeR=max(-hitMinR,hitMaxR);
+  AH2 lobeG=max(-hitMinG,hitMaxG);
+  AH2 lobeB=max(-hitMinB,hitMaxB);
+  AH2 lobe=max(AH2_(-FSR_RCAS_LIMIT),min(AMax3H2(lobeR,lobeG,lobeB),AH2_(0.0)))*AH2_(AH2_AU1(con.y).x);
+  // Apply noise removal.
+  #ifdef FSR_RCAS_DENOISE
+   lobe*=nz;
+  #endif
+  // Resolve, which needs the medium precision rcp approximation to avoid visible tonality changes.
+  AH2 rcpL=APrxMedRcpH2(AH2_(4.0)*lobe+AH2_(1.0));
+  pixR=(lobe*bR+lobe*dR+lobe*hR+lobe*fR+eR)*rcpL;
+  pixG=(lobe*bG+lobe*dG+lobe*hG+lobe*fG+eG)*rcpL;
+  pixB=(lobe*bB+lobe*dB+lobe*hB+lobe*fB+eB)*rcpL;}
+#endif
+////////////////////////////////////////////////////////////////////////////////////////////////////////////////////////////////
+////////////////////////////////////////////////////////////////////////////////////////////////////////////////////////////////
+////////////////////////////////////////////////////////////////////////////////////////////////////////////////////////////////
+////////////////////////////////////////////////////////////////////////////////////////////////////////////////////////////////
+//_____________________________________________________________/\_______________________________________________________________
+//==============================================================================================================================
+//
+//                                          FSR - [LFGA] LINEAR FILM GRAIN APPLICATOR
+//
+//------------------------------------------------------------------------------------------------------------------------------
+// Adding output-resolution film grain after scaling is a good way to mask both rendering and scaling artifacts.
+// Suggest using tiled blue noise as film grain input, with peak noise frequency set for a specific look and feel.
+// The 'Lfga*()' functions provide a convenient way to introduce grain.
+// These functions limit grain based on distance to signal limits.
+// This is done so that the grain is temporally energy preserving, and thus won't modify image tonality.
+// Grain application should be done in a linear colorspace.
+// The grain should be temporally changing, but have a temporal sum per pixel that adds to zero (non-biased).
+//------------------------------------------------------------------------------------------------------------------------------
+// Usage,
+//   FsrLfga*(
+//    color, // In/out linear colorspace color {0 to 1} ranged.
+//    grain, // Per pixel grain texture value {-0.5 to 0.5} ranged, input is 3-channel to support colored grain.
+//    amount); // Amount of grain (0 to 1} ranged.
+//------------------------------------------------------------------------------------------------------------------------------
+// Example if grain texture is monochrome: 'FsrLfgaF(color,AF3_(grain),amount)'
+//==============================================================================================================================
+#if defined(A_GPU)
+ // Maximum grain is the minimum distance to the signal limit.
+ void FsrLfgaF(inout AF3 c,AF3 t,AF1 a){c+=(t*AF3_(a))*min(AF3_(1.0)-c,c);}
+#endif
+//==============================================================================================================================
+#if defined(A_GPU)&&defined(A_HALF)
+ // Half precision version (slower).
+ void FsrLfgaH(inout AH3 c,AH3 t,AH1 a){c+=(t*AH3_(a))*min(AH3_(1.0)-c,c);}
+//------------------------------------------------------------------------------------------------------------------------------
+ // Packed half precision version (faster).
+ void FsrLfgaHx2(inout AH2 cR,inout AH2 cG,inout AH2 cB,AH2 tR,AH2 tG,AH2 tB,AH1 a){
+  cR+=(tR*AH2_(a))*min(AH2_(1.0)-cR,cR);cG+=(tG*AH2_(a))*min(AH2_(1.0)-cG,cG);cB+=(tB*AH2_(a))*min(AH2_(1.0)-cB,cB);}
+#endif
+////////////////////////////////////////////////////////////////////////////////////////////////////////////////////////////////
+////////////////////////////////////////////////////////////////////////////////////////////////////////////////////////////////
+////////////////////////////////////////////////////////////////////////////////////////////////////////////////////////////////
+////////////////////////////////////////////////////////////////////////////////////////////////////////////////////////////////
+//_____________________________________________________________/\_______________________________________________________________
+//==============================================================================================================================
+//
+//                                          FSR - [SRTM] SIMPLE REVERSIBLE TONE-MAPPER
+//
+//------------------------------------------------------------------------------------------------------------------------------
+// This provides a way to take linear HDR color {0 to FP16_MAX} and convert it into a temporary {0 to 1} ranged post-tonemapped linear.
+// The tonemapper preserves RGB ratio, which helps maintain HDR color bleed during filtering.
+//------------------------------------------------------------------------------------------------------------------------------
+// Reversible tonemapper usage,
+//  FsrSrtm*(color); // {0 to FP16_MAX} converted to {0 to 1}.
+//  FsrSrtmInv*(color); // {0 to 1} converted into {0 to 32768, output peak safe for FP16}.
+//==============================================================================================================================
+#if defined(A_GPU)
+ void FsrSrtmF(inout AF3 c){c*=AF3_(ARcpF1(AMax3F1(c.r,c.g,c.b)+AF1_(1.0)));}
+ // The extra max solves the c=1.0 case (which is a /0).
+ void FsrSrtmInvF(inout AF3 c){c*=AF3_(ARcpF1(max(AF1_(1.0/32768.0),AF1_(1.0)-AMax3F1(c.r,c.g,c.b))));}
+#endif
+//==============================================================================================================================
+#if defined(A_GPU)&&defined(A_HALF)
+ void FsrSrtmH(inout AH3 c){c*=AH3_(ARcpH1(AMax3H1(c.r,c.g,c.b)+AH1_(1.0)));}
+ void FsrSrtmInvH(inout AH3 c){c*=AH3_(ARcpH1(max(AH1_(1.0/32768.0),AH1_(1.0)-AMax3H1(c.r,c.g,c.b))));}
+//------------------------------------------------------------------------------------------------------------------------------
+ void FsrSrtmHx2(inout AH2 cR,inout AH2 cG,inout AH2 cB){
+  AH2 rcp=ARcpH2(AMax3H2(cR,cG,cB)+AH2_(1.0));cR*=rcp;cG*=rcp;cB*=rcp;}
+ void FsrSrtmInvHx2(inout AH2 cR,inout AH2 cG,inout AH2 cB){
+  AH2 rcp=ARcpH2(max(AH2_(1.0/32768.0),AH2_(1.0)-AMax3H2(cR,cG,cB)));cR*=rcp;cG*=rcp;cB*=rcp;}
+#endif
+////////////////////////////////////////////////////////////////////////////////////////////////////////////////////////////////
+////////////////////////////////////////////////////////////////////////////////////////////////////////////////////////////////
+////////////////////////////////////////////////////////////////////////////////////////////////////////////////////////////////
+////////////////////////////////////////////////////////////////////////////////////////////////////////////////////////////////
+//_____________________________________________________________/\_______________________________________________________________
+//==============================================================================================================================
+//
+//                                       FSR - [TEPD] TEMPORAL ENERGY PRESERVING DITHER
+//
+//------------------------------------------------------------------------------------------------------------------------------
+// Temporally energy preserving dithered {0 to 1} linear to gamma 2.0 conversion.
+// Gamma 2.0 is used so that the conversion back to linear is just to square the color.
+// The conversion comes in 8-bit and 10-bit modes, designed for output to 8-bit UNORM or 10:10:10:2 respectively.
+// Given good non-biased temporal blue noise as dither input,
+// the output dither will temporally conserve energy.
+// This is done by choosing the linear nearest step point instead of perceptual nearest.
+// See code below for details.
+//------------------------------------------------------------------------------------------------------------------------------
+// DX SPEC RULES FOR FLOAT->UNORM 8-BIT CONVERSION
+// ===============================================
+// - Output is 'uint(floor(saturate(n)*255.0+0.5))'.
+// - Thus rounding is to nearest.
+// - NaN gets converted to zero.
+// - INF is clamped to {0.0 to 1.0}.
+//==============================================================================================================================
+#if defined(A_GPU)
+ // Hand tuned integer position to dither value, with more values than simple checkerboard.
+ // Only 32-bit has enough precision for this compddation.
+ // Output is {0 to <1}.
+ AF1 FsrTepdDitF(AU2 p,AU1 f){
+  AF1 x=AF1_(p.x+f);
+  AF1 y=AF1_(p.y);
+  // The 1.61803 golden ratio.
+  AF1 a=AF1_((1.0+sqrt(5.0))/2.0);
+  // Number designed to provide a good visual pattern.
+  AF1 b=AF1_(1.0/3.69);
+  x=x*a+(y*b);
+  return AFractF1(x);}
+//------------------------------------------------------------------------------------------------------------------------------
+ // This version is 8-bit gamma 2.0.
+ // The 'c' input is {0 to 1}.
+ // Output is {0 to 1} ready for image store.
+ void FsrTepdC8F(inout AF3 c,AF1 dit){
+  AF3 n=sqrt(c);
+  n=floor(n*AF3_(255.0))*AF3_(1.0/255.0);
+  AF3 a=n*n;
+  AF3 b=n+AF3_(1.0/255.0);b=b*b;
+  // Ratio of 'a' to 'b' required to produce 'c'.
+  // APrxLoRcpF1() won't work here (at least for very high dynamic ranges).
+  // APrxMedRcpF1() is an IADD,FMA,MUL.
+  AF3 r=(c-b)*APrxMedRcpF3(a-b);
+  // Use the ratio as a cutoff to choose 'a' or 'b'.
+  // AGtZeroF1() is a MUL.
+  c=ASatF3(n+AGtZeroF3(AF3_(dit)-r)*AF3_(1.0/255.0));}
+//------------------------------------------------------------------------------------------------------------------------------
+ // This version is 10-bit gamma 2.0.
+ // The 'c' input is {0 to 1}.
+ // Output is {0 to 1} ready for image store.
+ void FsrTepdC10F(inout AF3 c,AF1 dit){
+  AF3 n=sqrt(c);
+  n=floor(n*AF3_(1023.0))*AF3_(1.0/1023.0);
+  AF3 a=n*n;
+  AF3 b=n+AF3_(1.0/1023.0);b=b*b;
+  AF3 r=(c-b)*APrxMedRcpF3(a-b);
+  c=ASatF3(n+AGtZeroF3(AF3_(dit)-r)*AF3_(1.0/1023.0));}
+#endif
+//==============================================================================================================================
+#if defined(A_GPU)&&defined(A_HALF)
+ AH1 FsrTepdDitH(AU2 p,AU1 f){
+  AF1 x=AF1_(p.x+f);
+  AF1 y=AF1_(p.y);
+  AF1 a=AF1_((1.0+sqrt(5.0))/2.0);
+  AF1 b=AF1_(1.0/3.69);
+  x=x*a+(y*b);
+  return AH1(AFractF1(x));}
+//------------------------------------------------------------------------------------------------------------------------------
+ void FsrTepdC8H(inout AH3 c,AH1 dit){
+  AH3 n=sqrt(c);
+  n=floor(n*AH3_(255.0))*AH3_(1.0/255.0);
+  AH3 a=n*n;
+  AH3 b=n+AH3_(1.0/255.0);b=b*b;
+  AH3 r=(c-b)*APrxMedRcpH3(a-b);
+  c=ASatH3(n+AGtZeroH3(AH3_(dit)-r)*AH3_(1.0/255.0));}
+//------------------------------------------------------------------------------------------------------------------------------
+ void FsrTepdC10H(inout AH3 c,AH1 dit){
+  AH3 n=sqrt(c);
+  n=floor(n*AH3_(1023.0))*AH3_(1.0/1023.0);
+  AH3 a=n*n;
+  AH3 b=n+AH3_(1.0/1023.0);b=b*b;
+  AH3 r=(c-b)*APrxMedRcpH3(a-b);
+  c=ASatH3(n+AGtZeroH3(AH3_(dit)-r)*AH3_(1.0/1023.0));}
+//==============================================================================================================================
+ // This computes dither for positions 'p' and 'p+{8,0}'.
+ AH2 FsrTepdDitHx2(AU2 p,AU1 f){
+  AF2 x;
+  x.x=AF1_(p.x+f);
+  x.y=x.x+AF1_(8.0);
+  AF1 y=AF1_(p.y);
+  AF1 a=AF1_((1.0+sqrt(5.0))/2.0);
+  AF1 b=AF1_(1.0/3.69);
+  x=x*AF2_(a)+AF2_(y*b);
+  return AH2(AFractF2(x));}
+//------------------------------------------------------------------------------------------------------------------------------
+ void FsrTepdC8Hx2(inout AH2 cR,inout AH2 cG,inout AH2 cB,AH2 dit){
+  AH2 nR=sqrt(cR);
+  AH2 nG=sqrt(cG);
+  AH2 nB=sqrt(cB);
+  nR=floor(nR*AH2_(255.0))*AH2_(1.0/255.0);
+  nG=floor(nG*AH2_(255.0))*AH2_(1.0/255.0);
+  nB=floor(nB*AH2_(255.0))*AH2_(1.0/255.0);
+  AH2 aR=nR*nR;
+  AH2 aG=nG*nG;
+  AH2 aB=nB*nB;
+  AH2 bR=nR+AH2_(1.0/255.0);bR=bR*bR;
+  AH2 bG=nG+AH2_(1.0/255.0);bG=bG*bG;
+  AH2 bB=nB+AH2_(1.0/255.0);bB=bB*bB;
+  AH2 rR=(cR-bR)*APrxMedRcpH2(aR-bR);
+  AH2 rG=(cG-bG)*APrxMedRcpH2(aG-bG);
+  AH2 rB=(cB-bB)*APrxMedRcpH2(aB-bB);
+  cR=ASatH2(nR+AGtZeroH2(dit-rR)*AH2_(1.0/255.0));
+  cG=ASatH2(nG+AGtZeroH2(dit-rG)*AH2_(1.0/255.0));
+  cB=ASatH2(nB+AGtZeroH2(dit-rB)*AH2_(1.0/255.0));}
+//------------------------------------------------------------------------------------------------------------------------------
+ void FsrTepdC10Hx2(inout AH2 cR,inout AH2 cG,inout AH2 cB,AH2 dit){
+  AH2 nR=sqrt(cR);
+  AH2 nG=sqrt(cG);
+  AH2 nB=sqrt(cB);
+  nR=floor(nR*AH2_(1023.0))*AH2_(1.0/1023.0);
+  nG=floor(nG*AH2_(1023.0))*AH2_(1.0/1023.0);
+  nB=floor(nB*AH2_(1023.0))*AH2_(1.0/1023.0);
+  AH2 aR=nR*nR;
+  AH2 aG=nG*nG;
+  AH2 aB=nB*nB;
+  AH2 bR=nR+AH2_(1.0/1023.0);bR=bR*bR;
+  AH2 bG=nG+AH2_(1.0/1023.0);bG=bG*bG;
+  AH2 bB=nB+AH2_(1.0/1023.0);bB=bB*bB;
+  AH2 rR=(cR-bR)*APrxMedRcpH2(aR-bR);
+  AH2 rG=(cG-bG)*APrxMedRcpH2(aG-bG);
+  AH2 rB=(cB-bB)*APrxMedRcpH2(aB-bB);
+  cR=ASatH2(nR+AGtZeroH2(dit-rR)*AH2_(1.0/1023.0));
+  cG=ASatH2(nG+AGtZeroH2(dit-rG)*AH2_(1.0/1023.0));
+  cB=ASatH2(nB+AGtZeroH2(dit-rB)*AH2_(1.0/1023.0));}
+#endif
diff --git a/Engine/Plugins/Runtime/MobileFSR/Source/MobileFSR/MobileFSR.Build.cs b/Engine/Plugins/Runtime/MobileFSR/Source/MobileFSR/MobileFSR.Build.cs
new file mode 100644
--- /dev/null
+++ b/Engine/Plugins/Runtime/MobileFSR/Source/MobileFSR/MobileFSR.Build.cs
@@ -0,0 +1,39 @@
+// Copyright Epic Games, Inc. All Rights Reserved.
+
+using UnrealBuildTool;
+using System.IO;
+
+public class MobileFSR : ModuleRules
+{
+	public MobileFSR(ReadOnlyTargetRules Target) : base(Target)
+	{
+		PCHUsage = ModuleRules.PCHUsageMode.UseExplicitOrSharedPCHs;
+
+		PrivateIncludePaths.AddRange(
+			new string[] {
+				EngineDirectory + "/Source/Runtime/Renderer/Private",
+			}
+			);
+
+		PrivateDependencyModuleNames.AddRange(
+			new string[]
+			{
+				"Core",
+				"Engine",
+				"Projects",
+				"RenderCore",
+				"Renderer",
+				"RHI",
+			}
+			);
+
+
+		PrecompileForTargets = PrecompileTargetsType.Any;
+
+		if (Target.Platform == UnrealTargetPlatform.Android)
+		{
+			//string PluginPath = Utils.MakePathRelativeTo(ModuleDirectory, Target.RelativeEnginePath);
+			//AdditionalPropertiesForReceipt.Add("AndroidPlugin", Path.Combine(PluginPath, "MobileFSR_APL.xml"));
+		}
+	}
+}
diff --git a/Engine/Plugins/Runtime/MobileFSR/Source/MobileFSR/MobileFSR_APL.xml b/Engine/Plugins/Runtime/MobileFSR/Source/MobileFSR/MobileFSR_APL.xml
new file mode 100644
--- /dev/null
+++ b/Engine/Plugins/Runtime/MobileFSR/Source/MobileFSR/MobileFSR_APL.xml
@@ -0,0 +1,15 @@
+<?xml version="1.0" encoding="utf-8"?>
+<!--MobileFSR plugin additions-->
+<root xmlns:android="http://schemas.android.com/apk/res/android">
+	<!-- init section is always evaluated once per architecture -->
+	<init>
+		<log text="MobileFSR init"/>
+		<setBool result="bSupported" value="false"/>
+		<isArch arch="armeabi-v7a">
+			<setBool result="bSupported" value="true"/>
+		</isArch>
+		<isArch arch="arm64-v8a">
+			<setBool result="bSupported" value="true"/>
+		</isArch>
+	</init>
+</root>
diff --git a/Engine/Plugins/Runtime/MobileFSR/Source/MobileFSR/Private/MobileFSRModule.cpp b/Engine/Plugins/Runtime/MobileFSR/Source/MobileFSR/Private/MobileFSRModule.cpp
new file mode 100644
--- /dev/null
+++ b/Engine/Plugins/Runtime/MobileFSR/Source/MobileFSR/Private/MobileFSRModule.cpp
@@ -0,0 +1,41 @@
+// Copyright Epic Games, Inc. All Rights Reserved.
+
+#include "MobileFSRModule.h"
+
+#include "CoreMinimal.h"
+#include "ShaderCore.h"
+#include "Interfaces/IPluginManager.h"
+#include "MobileFSRViewExtension.h"
+
+#define LOCTEXT_NAMESPACE "MobileFSR"
+
+DEFINE_LOG_CATEGORY(LogMobileFSR);
+
+void FMobileFSRModule::StartupModule()
+{
+	FString PluginShaderDir = FPaths::Combine(IPluginManager::Get().FindPlugin(TEXT("MobileFSR"))->GetBaseDir(), TEXT("Shaders"));
+	AddShaderSourceDirectoryMapping(TEXT("/Plugin/MobileFSR"), PluginShaderDir);
+
+	FCoreDelegates::OnPostEngineInit.AddRaw(this, &FMobileFSRModule::OnPostEngineInit);
+
+	UE_LOG(LogMobileFSR, Log, TEXT("Mobile FSR Started"));
+}
+
+void FMobileFSRModule::OnPostEngineInit()
+{
+	ViewExtension = FSceneViewExtensions::NewExtension<FMobileFSRViewExtension>();
+
+	UE_LOG(LogMobileFSR, Log, TEXT("Mobile FSR OnPostEngineInit"));
+}
+
+void FMobileFSRModule::ShutdownModule()
+{
+	ViewExtension = nullptr;
+	FCoreDelegates::OnPostEngineInit.RemoveAll(this);
+
+	UE_LOG(LogMobileFSR, Log, TEXT("Mobile FSR Shutdown"));
+}
+
+#undef LOCTEXT_NAMESPACE
+
+IMPLEMENT_MODULE(FMobileFSRModule, MobileFSR)
\ No newline at end of file
diff --git a/Engine/Plugins/Runtime/MobileFSR/Source/MobileFSR/Private/MobileFSRUpscaler.cpp b/Engine/Plugins/Runtime/MobileFSR/Source/MobileFSR/Private/MobileFSRUpscaler.cpp
new file mode 100644
--- /dev/null
+++ b/Engine/Plugins/Runtime/MobileFSR/Source/MobileFSR/Private/MobileFSRUpscaler.cpp
@@ -0,0 +1,445 @@
+// Copyright Epic Games, Inc. All Rights Reserved.
+
+#include "MobileFSRUpscaler.h"
+#include "SceneRenderTargetParameters.h"
+#include "SceneRendering.h"
+#include "PostProcess/PostProcessing.h"
+
+#define A_CPU 1
+// ffx_a.ush will override these definitions
+#ifdef A_STATIC
+#undef A_STATIC
+#endif
+#ifdef A_RESTRICT
+#undef A_RESTRICT
+#endif
+#include "../Shaders/Private/ffx_a.ush"
+#include "../Shaders/Private/ffx_fsr1.ush"
+
+
+//////////////////////////////////////////////////////////////////////////
+//! FSR
+//////////////////////////////////////////////////////////////////////////
+
+// permutation domains
+class FMobileFSR_UseFP16Dim : SHADER_PERMUTATION_BOOL("ENABLE_FP16");
+class FMobileRCAS_DenoiseDim : SHADER_PERMUTATION_BOOL("USE_RCAS_DENOISE");
+
+BEGIN_SHADER_PARAMETER_STRUCT(FMobileFSRPassParameters, )
+SHADER_PARAMETER(FUintVector4, Const0)
+SHADER_PARAMETER(FUintVector4, Const1)
+SHADER_PARAMETER(FUintVector4, Const2)
+SHADER_PARAMETER(FUintVector4, Const3)
+SHADER_PARAMETER(FVector2f, VPColor_ExtentInverse)
+SHADER_PARAMETER(FVector2f, VPColor_ViewportMin)
+SHADER_PARAMETER_SAMPLER(SamplerState, samLinearClamp)
+SHADER_PARAMETER_RDG_TEXTURE(Texture2D, InputTexture)
+END_SHADER_PARAMETER_STRUCT()
+
+BEGIN_SHADER_PARAMETER_STRUCT(FMobileRCASPassParameters, )
+SHADER_PARAMETER_RDG_TEXTURE(Texture2D, InputTexture)
+SHADER_PARAMETER(FUintVector4, Const0)
+SHADER_PARAMETER(FVector2f, VPColor_ExtentInverse)
+END_SHADER_PARAMETER_STRUCT()
+
+static float GMobileFSRCASSharpness = 0.2; // 0.2 stops = 1 / (2^N) ~= 0.87 in linear [0-1] sharpening of CAS
+static FAutoConsoleVariableRef CVarMobileFSRCASSharpness(
+	TEXT("r.Mobile.FSR.RCAS.Sharpness"),
+	GMobileFSRCASSharpness,
+	TEXT("FSR RCAS Sharpness in stops (0: sharpest, 1: 1/2 as sharp, 2: 1/4 as sharp, 3: 1/8 as sharp, etc.). A value of 0.2 would correspond to a ~0.87 sharpness in [0-1] linear scale"),
+	ECVF_Default);
+
+static int32 GMobileFSRCASDenoise = 0;
+static FAutoConsoleVariableRef CVarMobileFSRCASDenoise(
+	TEXT("r.Mobile.FSR.RCAS.Denoise"),
+	GMobileFSRCASDenoise,
+	TEXT("FSR RCAS Denoise support for grainy input such as dithered images or input with custom film grain effects applied prior to FSR. 1:On, 0:Off"),
+	ECVF_Default);
+
+static int32 GMobileFSRDisableCompute = 0;
+static FAutoConsoleVariableRef CVarMobileFSRDisableCompute(
+	TEXT("r.Mobile.FSR.DisableCompute"),
+	GMobileFSRDisableCompute,
+	TEXT("FSR Disable compute shaders for FSR (Recommended for Adreno). 1: Compute Disabled, 0: Compute Enabled (default)"),
+	ECVF_Default);
+
+///
+/// FSR COMPUTE SHADER
+///
+class FMobileFSRCS : public FGlobalShader
+{
+public:
+	DECLARE_GLOBAL_SHADER(FMobileFSRCS);
+	SHADER_USE_PARAMETER_STRUCT(FMobileFSRCS, FGlobalShader);
+
+	using FPermutationDomain = TShaderPermutationDomain<FMobileFSR_UseFP16Dim>;
+
+	BEGIN_SHADER_PARAMETER_STRUCT(FParameters, )
+		SHADER_PARAMETER_STRUCT_INCLUDE(FMobileFSRPassParameters, FSR)
+		SHADER_PARAMETER_RDG_TEXTURE_UAV(RWTexture2D, OutputTexture)
+	END_SHADER_PARAMETER_STRUCT()
+	
+	static bool ShouldCompilePermutation(const FGlobalShaderPermutationParameters& Parameters)
+	{
+		FPermutationDomain PermutationVector(Parameters.PermutationId);
+
+		if (PermutationVector.Get<FMobileFSR_UseFP16Dim>())
+		{
+			if (IsSimulatedPlatform(Parameters.Platform) || !IsDxcEnabledForPlatform(Parameters.Platform))
+			{
+				return false;
+			}
+		}
+
+		return IsMobilePlatform(Parameters.Platform);
+	}
+
+	static void ModifyCompilationEnvironment(const FGlobalShaderPermutationParameters& Parameters, FShaderCompilerEnvironment& OutEnvironment)
+	{
+		FPermutationDomain PermutationVector(Parameters.PermutationId);
+
+		OutEnvironment.SetDefine(TEXT("COMPUTE_SHADER"), 1);
+
+		OutEnvironment.SetDefine(TEXT("THREADGROUP_SIZEX"), 64);
+		OutEnvironment.SetDefine(TEXT("THREADGROUP_SIZEY"), 1);
+		OutEnvironment.SetDefine(TEXT("THREADGROUP_SIZEZ"), 1);
+	}
+};
+
+IMPLEMENT_GLOBAL_SHADER(FMobileFSRCS, "/Plugin/MobileFSR/Private/PostProcessMobileFFX_FSR.usf", "MainCS", SF_Compute);
+
+class FMobileFSRPS : public FGlobalShader
+{
+public:
+	DECLARE_GLOBAL_SHADER(FMobileFSRPS);
+	SHADER_USE_PARAMETER_STRUCT(FMobileFSRPS, FGlobalShader);
+
+	using FPermutationDomain = TShaderPermutationDomain<FMobileFSR_UseFP16Dim>;
+
+	BEGIN_SHADER_PARAMETER_STRUCT(FParameters, )
+		SHADER_PARAMETER_STRUCT_INCLUDE(FMobileFSRPassParameters, FSR)
+		RENDER_TARGET_BINDING_SLOTS()
+	END_SHADER_PARAMETER_STRUCT()
+
+	static bool ShouldCompilePermutation(const FGlobalShaderPermutationParameters& Parameters)
+	{
+		FPermutationDomain PermutationVector(Parameters.PermutationId);
+
+		if (PermutationVector.Get<FMobileFSR_UseFP16Dim>())
+		{
+			if (IsSimulatedPlatform(Parameters.Platform) || !IsDxcEnabledForPlatform(Parameters.Platform))
+			{
+				return false;
+			}
+		}
+
+		return IsMobilePlatform(Parameters.Platform);
+	}
+
+	static void ModifyCompilationEnvironment(const FGlobalShaderPermutationParameters& Parameters, FShaderCompilerEnvironment& OutEnvironment)
+	{
+	}
+};
+
+IMPLEMENT_GLOBAL_SHADER(FMobileFSRPS, "/Plugin/MobileFSR/Private/PostProcessMobileFFX_FSR.usf", "MainPS", SF_Pixel);
+
+FScreenPassTexture AddEASUPass(FRDGBuilder& GraphBuilder, const FViewInfo& View, const ISpatialUpscaler::FInputs& PassInputs)
+{
+	FScreenPassTexture OutputTexture;
+	bool bUseFP16 = !IsSimulatedPlatform(View.GetShaderPlatform()) && IsDxcEnabledForPlatform(View.GetShaderPlatform());
+
+	FRDGTextureDesc FSROutputTextureDesc = PassInputs.SceneColor.Texture->Desc;
+	FSROutputTextureDesc.Reset();
+	FSROutputTextureDesc.Extent = View.UnscaledViewRect.Max;
+	FSROutputTextureDesc.ClearValue = FClearValueBinding::Black;
+	FSROutputTextureDesc.Flags = TexCreate_ShaderResource | TexCreate_UAV | TexCreate_RenderTargetable;
+
+	if (PassInputs.OverrideOutput.IsValid())
+	{
+		OutputTexture = PassInputs.OverrideOutput;
+	}
+	else
+	{
+		OutputTexture = FScreenPassTexture(PassInputs.SceneColor);
+		OutputTexture.Texture = GraphBuilder.CreateTexture(FSROutputTextureDesc, TEXT("FSR-Output"), ERDGTextureFlags::MultiFrame);
+	}
+
+	OutputTexture.ViewRect = View.UnscaledViewRect;
+	FScreenPassTextureViewport OutputViewport = FScreenPassTextureViewport(OutputTexture);
+	FScreenPassTextureViewport InputViewport = FScreenPassTextureViewport(PassInputs.SceneColor);
+	FScreenPassTextureViewportParameters PassOutputViewportParams = GetScreenPassTextureViewportParameters(OutputViewport);
+
+	AU1 const0[4];
+	AU1 const1[4];
+	AU1 const2[4];
+	AU1 const3[4]; // Configure FSR
+
+	FsrEasuConOffset(const0, const1, const2, const3,
+		static_cast<AF1>(InputViewport.Rect.Width())
+		, static_cast<AF1>(InputViewport.Rect.Height()) // current frame render resolution
+		, static_cast<AF1>(PassInputs.SceneColor.Texture->Desc.Extent.X)
+		, static_cast<AF1>(PassInputs.SceneColor.Texture->Desc.Extent.Y) // input container resolution (for DRS)
+		, static_cast<AF1>(OutputViewport.Rect.Width())
+		, static_cast<AF1>(OutputViewport.Rect.Height()) // upscaled-to resolution
+		, static_cast<AF1>(InputViewport.Rect.Min.X)
+		, static_cast<AF1>(InputViewport.Rect.Min.Y)
+	);
+
+	const bool bOutputSupportsUAV = (OutputTexture.Texture->Desc.Flags & TexCreate_UAV) == TexCreate_UAV;
+
+	if (!bOutputSupportsUAV || GMobileFSRDisableCompute)
+	{
+		FMobileFSRPS::FParameters* PassParameters = GraphBuilder.AllocParameters<FMobileFSRPS::FParameters>();
+		for (int i = 0; i < 4; i++)
+		{
+			PassParameters->FSR.Const0[i] = const0[i];
+			PassParameters->FSR.Const1[i] = const1[i];
+			PassParameters->FSR.Const2[i] = const2[i];
+			PassParameters->FSR.Const3[i] = const3[i];
+		}
+		PassParameters->FSR.InputTexture = PassInputs.SceneColor.Texture;
+		PassParameters->FSR.samLinearClamp = TStaticSamplerState<SF_Bilinear, AM_Clamp, AM_Clamp, AM_Clamp>::GetRHI();
+		PassParameters->FSR.VPColor_ExtentInverse = PassOutputViewportParams.ExtentInverse;
+		PassParameters->FSR.VPColor_ViewportMin = PassOutputViewportParams.ViewportMin;
+		PassParameters->RenderTargets[0] = FRenderTargetBinding(OutputTexture.Texture, ERenderTargetLoadAction::EClear);
+
+		FMobileFSRPS::FPermutationDomain PSPermutationVector;
+		PSPermutationVector.Set<FMobileFSR_UseFP16Dim>(bUseFP16);
+
+		TShaderMapRef<FMobileFSRPS> PixelShader(View.ShaderMap, PSPermutationVector);
+
+		AddDrawScreenPass(GraphBuilder,
+			RDG_EVENT_NAME("FSR/Upscale %dx%d -> %dx%d (PS)"
+				, InputViewport.Rect.Width(), InputViewport.Rect.Height()
+				, OutputViewport.Rect.Width(), OutputViewport.Rect.Height()),
+			View, OutputViewport, InputViewport,
+			PixelShader, PassParameters,
+			EScreenPassDrawFlags::None
+		);
+	}
+	// CS
+	else
+	{
+		FMobileFSRCS::FParameters* PassParameters = GraphBuilder.AllocParameters<FMobileFSRCS::FParameters>();
+
+		for (int i = 0; i < 4; i++)
+		{
+			PassParameters->FSR.Const0[i] = const0[i];
+			PassParameters->FSR.Const1[i] = const1[i];
+			PassParameters->FSR.Const2[i] = const2[i];
+			PassParameters->FSR.Const3[i] = const3[i];
+		}
+		PassParameters->FSR.InputTexture = PassInputs.SceneColor.Texture;
+		PassParameters->FSR.samLinearClamp = TStaticSamplerState<SF_Bilinear, AM_Clamp, AM_Clamp, AM_Clamp>::GetRHI();
+		PassParameters->FSR.VPColor_ExtentInverse = PassOutputViewportParams.ExtentInverse;
+		PassParameters->FSR.VPColor_ViewportMin = PassOutputViewportParams.ViewportMin;
+		PassParameters->OutputTexture = GraphBuilder.CreateUAV(OutputTexture.Texture);
+
+		FMobileFSRCS::FPermutationDomain CSPermutationVector;
+		CSPermutationVector.Set<FMobileFSR_UseFP16Dim>(bUseFP16);
+
+		TShaderMapRef<FMobileFSRCS> ComputeShaderFSR(View.ShaderMap, CSPermutationVector);
+		FComputeShaderUtils::AddPass(
+			GraphBuilder,
+			RDG_EVENT_NAME("FSR/Upscale %dx%d -> %dx%d (CS)"
+				, InputViewport.Rect.Width(), InputViewport.Rect.Height()
+				, OutputViewport.Rect.Width(), OutputViewport.Rect.Height()),
+			ComputeShaderFSR,
+			PassParameters,
+			FComputeShaderUtils::GetGroupCount(OutputViewport.Rect.Size(), 16));
+	}
+
+	return MoveTemp(OutputTexture);
+}
+
+///
+/// RCAS COMPUTE SHADER
+///
+class FMobileRCASCS : public FGlobalShader
+{
+public:
+	DECLARE_GLOBAL_SHADER(FMobileRCASCS);
+	SHADER_USE_PARAMETER_STRUCT(FMobileRCASCS, FGlobalShader);
+
+	using FPermutationDomain = TShaderPermutationDomain<FMobileFSR_UseFP16Dim, FMobileRCAS_DenoiseDim>;
+
+	BEGIN_SHADER_PARAMETER_STRUCT(FParameters, )
+		SHADER_PARAMETER_STRUCT_INCLUDE(FMobileRCASPassParameters, RCAS)
+		SHADER_PARAMETER_RDG_TEXTURE_UAV(RWTexture2D, OutputTexture)
+	END_SHADER_PARAMETER_STRUCT()
+
+	static bool ShouldCompilePermutation(const FGlobalShaderPermutationParameters& Parameters)
+	{
+		FPermutationDomain PermutationVector(Parameters.PermutationId);
+
+		if (PermutationVector.Get<FMobileFSR_UseFP16Dim>())
+		{
+			if (IsSimulatedPlatform(Parameters.Platform) || !IsDxcEnabledForPlatform(Parameters.Platform))
+			{
+				return false;
+			}
+		}
+
+		return IsMobilePlatform(Parameters.Platform);
+	}
+
+	static void ModifyCompilationEnvironment(const FGlobalShaderPermutationParameters& Parameters, FShaderCompilerEnvironment& OutEnvironment)
+	{
+		FPermutationDomain PermutationVector(Parameters.PermutationId);
+
+		OutEnvironment.SetDefine(TEXT("COMPUTE_SHADER"), 1);
+
+		OutEnvironment.SetDefine(TEXT("THREADGROUP_SIZEX"), 64);
+		OutEnvironment.SetDefine(TEXT("THREADGROUP_SIZEY"), 1);
+		OutEnvironment.SetDefine(TEXT("THREADGROUP_SIZEZ"), 1);
+	}
+};
+
+IMPLEMENT_GLOBAL_SHADER(FMobileRCASCS, "/Plugin/MobileFSR/Private/PostProcessMobileFFX_RCAS.usf", "MainCS", SF_Compute);
+
+
+///
+/// RCAS PIXEL SHADER
+///
+class FMobileRCASPS : public FGlobalShader
+{
+public:
+	DECLARE_GLOBAL_SHADER(FMobileRCASPS);
+	SHADER_USE_PARAMETER_STRUCT(FMobileRCASPS, FGlobalShader);
+
+	using FPermutationDomain = TShaderPermutationDomain<FMobileFSR_UseFP16Dim, FMobileRCAS_DenoiseDim>;
+
+	BEGIN_SHADER_PARAMETER_STRUCT(FParameters, )
+		SHADER_PARAMETER_STRUCT_INCLUDE(FMobileRCASPassParameters, RCAS)
+		RENDER_TARGET_BINDING_SLOTS()
+	END_SHADER_PARAMETER_STRUCT()
+
+	static bool ShouldCompilePermutation(const FGlobalShaderPermutationParameters& Parameters)
+	{
+		FPermutationDomain PermutationVector(Parameters.PermutationId);
+		if (PermutationVector.Get<FMobileFSR_UseFP16Dim>())
+		{
+			if (IsSimulatedPlatform(Parameters.Platform) || !IsDxcEnabledForPlatform(Parameters.Platform))
+			{
+				return false;
+			}
+		}
+
+		return IsMobilePlatform(Parameters.Platform);
+	}
+
+	static void ModifyCompilationEnvironment(const FGlobalShaderPermutationParameters& Parameters, FShaderCompilerEnvironment& OutEnvironment)
+	{
+	}
+};
+
+IMPLEMENT_GLOBAL_SHADER(FMobileRCASPS, "/Plugin/MobileFSR/Private/PostProcessMobileFFX_RCAS.usf", "MainPS", SF_Pixel);
+
+FScreenPassTexture AddCASPass(FRDGBuilder& GraphBuilder, const FViewInfo& View, const ISpatialUpscaler::FInputs& PassInputs)
+{
+	FScreenPassTexture OutputTexture;
+	bool bUseFP16 = !IsSimulatedPlatform(View.GetShaderPlatform()) && IsDxcEnabledForPlatform(View.GetShaderPlatform());
+
+	if (PassInputs.OverrideOutput.IsValid())
+	{
+		OutputTexture = PassInputs.OverrideOutput;
+	}
+	else
+	{
+		FRDGTextureDesc RCASOutputTextureDesc = PassInputs.SceneColor.Texture->Desc;
+
+		OutputTexture = FScreenPassTexture(PassInputs.SceneColor);
+		OutputTexture.Texture = GraphBuilder.CreateTexture(RCASOutputTextureDesc, TEXT("RCAS-Output"), ERDGTextureFlags::MultiFrame);
+	}
+
+	FScreenPassTextureViewport OutputViewport = FScreenPassTextureViewport(OutputTexture);
+	FScreenPassTextureViewportParameters PassOutputViewportParams = GetScreenPassTextureViewportParameters(OutputViewport);
+
+	AU1 const0[4]; // Configure FSR
+	FsrRcasCon(const0, GMobileFSRCASSharpness);
+
+	const bool bOutputSupportsUAV = (OutputTexture.Texture->Desc.Flags & TexCreate_UAV) == TexCreate_UAV;
+
+	if (!bOutputSupportsUAV || GMobileFSRDisableCompute)
+	{
+		FMobileRCASPS::FParameters* PassParameters = GraphBuilder.AllocParameters<FMobileRCASPS::FParameters>();
+
+		// set pass inputs
+		for (int i = 0; i < 4; i++) { PassParameters->RCAS.Const0[i] = const0[i]; }
+		PassParameters->RCAS.InputTexture = PassInputs.SceneColor.Texture;
+		PassParameters->RCAS.VPColor_ExtentInverse = PassOutputViewportParams.ExtentInverse;
+		PassParameters->RenderTargets[0] = FRenderTargetBinding(OutputTexture.Texture, ERenderTargetLoadAction::EClear);
+
+		// grab shaders
+		FMobileRCASPS::FPermutationDomain PSPermutationVector;
+		PSPermutationVector.Set<FMobileFSR_UseFP16Dim>(bUseFP16);
+		PSPermutationVector.Set<FMobileRCAS_DenoiseDim>(GMobileFSRCASDenoise > 0);
+
+		TShaderMapRef<FMobileRCASPS> PixelShader(View.ShaderMap, PSPermutationVector);
+
+		AddDrawScreenPass(GraphBuilder,
+			RDG_EVENT_NAME("FSR/RCAS Sharpness=%.2f(PS)"
+				, GMobileFSRCASSharpness),
+			View, OutputViewport, OutputViewport,
+			PixelShader, PassParameters,
+			EScreenPassDrawFlags::None
+		);
+	}
+	// CS
+	else
+	{
+		FMobileRCASCS::FParameters* PassParameters = GraphBuilder.AllocParameters<FMobileRCASCS::FParameters>();
+
+		for (int i = 0; i < 4; i++)
+		{
+			PassParameters->RCAS.Const0[i] = const0[i];
+		}
+		PassParameters->RCAS.InputTexture = PassInputs.SceneColor.Texture;
+		PassParameters->RCAS.VPColor_ExtentInverse = PassOutputViewportParams.ExtentInverse;
+		PassParameters->OutputTexture = GraphBuilder.CreateUAV(OutputTexture.Texture);
+
+		FMobileRCASCS::FPermutationDomain CSPermutationVector;
+		CSPermutationVector.Set<FMobileFSR_UseFP16Dim>(bUseFP16);
+		CSPermutationVector.Set<FMobileRCAS_DenoiseDim>(GMobileFSRCASDenoise > 0);
+
+		TShaderMapRef<FMobileRCASCS> ComputeShaderRCASPass(View.ShaderMap, CSPermutationVector);
+		FComputeShaderUtils::AddPass(
+			GraphBuilder,
+			RDG_EVENT_NAME("FSR/RCAS Sharpness=%.2f (CS)", GMobileFSRCASSharpness),
+			ComputeShaderRCASPass,
+			PassParameters,
+			FComputeShaderUtils::GetGroupCount(OutputViewport.Rect.Size(), 16));
+	}
+
+	return MoveTemp(OutputTexture);
+}
+
+DECLARE_GPU_STAT(MobileFSRUpscaler)
+
+FMobileFSRUpscaler::FMobileFSRUpscaler(bool bInIsEASUPass) : bIsEASUPass(bInIsEASUPass)
+{
+}
+
+ISpatialUpscaler* FMobileFSRUpscaler::Fork_GameThread(const class FSceneViewFamily& ViewFamily) const
+{
+	return new FMobileFSRUpscaler(bIsEASUPass);
+}
+
+FScreenPassTexture FMobileFSRUpscaler::AddPasses(FRDGBuilder& GraphBuilder, const FViewInfo& View, const FInputs& PassInputs) const
+{
+	RDG_GPU_STAT_SCOPE(GraphBuilder, MobileFSRUpscaler);
+	
+	FScreenPassTexture SceneColor;
+
+	if (bIsEASUPass)
+	{
+		SceneColor = AddEASUPass(GraphBuilder, View, PassInputs);
+	}
+	else
+	{
+		SceneColor = AddCASPass(GraphBuilder, View, PassInputs);
+	}
+
+	return MoveTemp(SceneColor);
+
+}
\ No newline at end of file
diff --git a/Engine/Plugins/Runtime/MobileFSR/Source/MobileFSR/Private/MobileFSRUpscaler.h b/Engine/Plugins/Runtime/MobileFSR/Source/MobileFSR/Private/MobileFSRUpscaler.h
new file mode 100644
--- /dev/null
+++ b/Engine/Plugins/Runtime/MobileFSR/Source/MobileFSR/Private/MobileFSRUpscaler.h
@@ -0,0 +1,20 @@
+// Copyright Epic Games, Inc. All Rights Reserved.
+
+#pragma once
+
+#include "PostProcess/PostProcessUpscale.h"
+
+class FMobileFSRUpscaler final : public ISpatialUpscaler
+{
+public:
+	FMobileFSRUpscaler(bool bInIsEASUPass);
+
+	// ISpatialUpscaler interface
+	const TCHAR* GetDebugName() const override { return TEXT("FMobileFSRUpscaler"); }
+
+	ISpatialUpscaler* Fork_GameThread(const class FSceneViewFamily& ViewFamily) const override;
+	FScreenPassTexture AddPasses(FRDGBuilder& GraphBuilder, const FViewInfo& View, const FInputs& PassInputs) const override;
+
+private:
+	bool bIsEASUPass = false;
+};
\ No newline at end of file
diff --git a/Engine/Plugins/Runtime/MobileFSR/Source/MobileFSR/Private/MobileFSRViewExtension.cpp b/Engine/Plugins/Runtime/MobileFSR/Source/MobileFSR/Private/MobileFSRViewExtension.cpp
new file mode 100644
--- /dev/null
+++ b/Engine/Plugins/Runtime/MobileFSR/Source/MobileFSR/Private/MobileFSRViewExtension.cpp
@@ -0,0 +1,38 @@
+// Copyright Epic Games, Inc. All Rights Reserved.
+
+#include "MobileFSRViewExtension.h"
+#include "MobileFSRUpscaler.h"
+#include "PostProcess/PostProcessUpscale.h"
+
+static TAutoConsoleVariable<int32> CVarEnableFSR(
+	TEXT("r.Mobile.FSR.Enabled"),
+	1,
+	TEXT("Enable Mobile FSR for Primary Upscale"),
+	ECVF_RenderThreadSafe);
+
+TAutoConsoleVariable<int32> CVarMobileFSRCASEnabled(
+	TEXT("r.Mobile.FSR.RCAS.Enabled"),
+	1,
+	TEXT("FidelityFX FSR/RCAS : Robust Contrast Adaptive Sharpening Filter"),
+	ECVF_RenderThreadSafe);
+
+TAutoConsoleVariable<int32>CVarMobileFSRUpsamplingEnabled(
+	TEXT("r.Mobile.FSR.Upsampling.Enabled"),
+	1,
+	TEXT("FidelityFX FSR/RCAS : Robust Contrast Adaptive Sharpening Filter"),
+	ECVF_Default);
+  
+void FMobileFSRViewExtension::BeginRenderViewFamily(FSceneViewFamily& InViewFamily)
+{
+	if (InViewFamily.GetFeatureLevel() == ERHIFeatureLevel::ES3_1 && CVarEnableFSR.GetValueOnAnyThread() > 0)
+	{
+		if (CVarMobileFSRUpsamplingEnabled.GetValueOnAnyThread())
+		{
+			InViewFamily.SetPrimarySpatialUpscalerInterface(new FMobileFSRUpscaler(true));
+		}
+		if (CVarMobileFSRCASEnabled.GetValueOnAnyThread())
+		{
+			InViewFamily.SetSecondarySpatialUpscalerInterface(new FMobileFSRUpscaler(false));
+		}
+	}
+}
\ No newline at end of file
diff --git a/Engine/Plugins/Runtime/MobileFSR/Source/MobileFSR/Public/MobileFSRModule.h b/Engine/Plugins/Runtime/MobileFSR/Source/MobileFSR/Public/MobileFSRModule.h
new file mode 100644
--- /dev/null
+++ b/Engine/Plugins/Runtime/MobileFSR/Source/MobileFSR/Public/MobileFSRModule.h
@@ -0,0 +1,20 @@
+// Copyright Epic Games, Inc. All Rights Reserved.
+
+#pragma once
+
+#include "Modules/ModuleManager.h"
+#include "Logging/LogMacros.h"
+
+DECLARE_LOG_CATEGORY_EXTERN(LogMobileFSR, Verbose, All);
+
+class FMobileFSRViewExtension;
+
+class FMobileFSRModule final : public IModuleInterface
+{
+public:
+	void StartupModule() override;
+	void ShutdownModule() override;
+	void OnPostEngineInit();
+
+	TSharedPtr<FMobileFSRViewExtension, ESPMode::ThreadSafe> ViewExtension;
+};
\ No newline at end of file
diff --git a/Engine/Plugins/Runtime/MobileFSR/Source/MobileFSR/Public/MobileFSRViewExtension.h b/Engine/Plugins/Runtime/MobileFSR/Source/MobileFSR/Public/MobileFSRViewExtension.h
new file mode 100644
--- /dev/null
+++ b/Engine/Plugins/Runtime/MobileFSR/Source/MobileFSR/Public/MobileFSRViewExtension.h
@@ -0,0 +1,18 @@
+// Copyright Epic Games, Inc. All Rights Reserved.
+
+#pragma once
+
+#include "SceneViewExtension.h"
+
+class FMobileFSRViewExtension final : public FSceneViewExtensionBase
+{
+public:
+	FMobileFSRViewExtension(const FAutoRegister& AutoRegister) : FSceneViewExtensionBase(AutoRegister) {}
+
+	void SetupViewFamily(FSceneViewFamily& InViewFamily) override {}
+	void SetupView(FSceneViewFamily& InViewFamily, FSceneView& InView) override {}
+	void PreRenderViewFamily_RenderThread(FRHICommandListImmediate& RHICmdList, FSceneViewFamily& InViewFamily) override {}
+		
+	MOBILEFSR_API void BeginRenderViewFamily(FSceneViewFamily& InViewFamily) override;
+
+};
\ No newline at end of file
