MobileFSR: Use proper include path for private shaders

#rnx

diff --git a/Engine/Plugins/Runtime/MobileFSR/Source/MobileFSR/MobileFSR.Build.cs b/Engine/Plugins/Runtime/MobileFSR/Source/MobileFSR/MobileFSR.Build.cs
--- a/Engine/Plugins/Runtime/MobileFSR/Source/MobileFSR/MobileFSR.Build.cs
+++ b/Engine/Plugins/Runtime/MobileFSR/Source/MobileFSR/MobileFSR.Build.cs
@@ -1,7 +1,6 @@
 // Copyright Epic Games, Inc. All Rights Reserved.
 
 using UnrealBuildTool;
-using System.IO;
 
 public class MobileFSR : ModuleRules
 {
@@ -12,6 +11,7 @@ public class MobileFSR : ModuleRules
 		PrivateIncludePaths.AddRange(
 			new string[] {
 				EngineDirectory + "/Source/Runtime/Renderer/Private",
+				PluginDirectory + "/Shaders/Private",
 			}
 			);
 
diff --git a/Engine/Plugins/Runtime/MobileFSR/Source/MobileFSR/Private/MobileFSRUpscaler.cpp b/Engine/Plugins/Runtime/MobileFSR/Source/MobileFSR/Private/MobileFSRUpscaler.cpp
--- a/Engine/Plugins/Runtime/MobileFSR/Source/MobileFSR/Private/MobileFSRUpscaler.cpp
+++ b/Engine/Plugins/Runtime/MobileFSR/Source/MobileFSR/Private/MobileFSRUpscaler.cpp
@@ -13,8 +13,8 @@
 #ifdef A_RESTRICT
 #undef A_RESTRICT
 #endif
-#include "../Shaders/Private/ffx_a.ush"
-#include "../Shaders/Private/ffx_fsr1.ush"
+#include "ffx_a.ush"
+#include "ffx_fsr1.ush"
 
 
 //////////////////////////////////////////////////////////////////////////
