Add Android as only SupportedPlatform in MobileFSR plugin

diff --git a/Engine/Plugins/Runtime/MobileFSR/MobileFSR.uplugin b/Engine/Plugins/Runtime/MobileFSR/MobileFSR.uplugin
--- a/Engine/Plugins/Runtime/MobileFSR/MobileFSR.uplugin
+++ b/Engine/Plugins/Runtime/MobileFSR/MobileFSR.uplugin
@@ -14,6 +14,9 @@
 	"CanContainContent": true,
 	"IsBetaVersion": false,
 	"Installed": false,
+	"SupportedTargetPlatforms": [
+		"Android"
+	],
 	"Modules": [
 		{
 			"Name": "MobileFSR",
