Add IOS as supported platform for MobileFSR

diff --git a/Engine/Plugins/Runtime/MobileFSR/MobileFSR.uplugin b/Engine/Plugins/Runtime/MobileFSR/MobileFSR.uplugin
--- a/Engine/Plugins/Runtime/MobileFSR/MobileFSR.uplugin
+++ b/Engine/Plugins/Runtime/MobileFSR/MobileFSR.uplugin
@@ -15,6 +15,7 @@
 	"IsBetaVersion": false,
 	"Installed": false,
 	"SupportedTargetPlatforms": [
+		"IOS",
 		"Android"
 	],
 	"Modules": [
