Add Alpha Passthrough Support to Mobile FSR

diff --git a/Engine/Plugins/Runtime/MobileFSR/Shaders/Private/PostProcessMobileFFX_FSR.usf b/Engine/Plugins/Runtime/MobileFSR/Shaders/Private/PostProcessMobileFFX_FSR.usf
--- a/Engine/Plugins/Runtime/MobileFSR/Shaders/Private/PostProcessMobileFFX_FSR.usf
+++ b/Engine/Plugins/Runtime/MobileFSR/Shaders/Private/PostProcessMobileFFX_FSR.usf
@@ -62,10 +62,12 @@ float2 VPColor_ViewportMin;
 	AH4 FsrEasuRH(AF2 p) { AH4 res = InputTexture.GatherRed  (samLinearClamp, p, ASU2(0, 0)); return res; }
 	AH4 FsrEasuGH(AF2 p) { AH4 res = InputTexture.GatherGreen(samLinearClamp, p, ASU2(0, 0)); return res; }
 	AH4 FsrEasuBH(AF2 p) { AH4 res = InputTexture.GatherBlue (samLinearClamp, p, ASU2(0, 0)); return res; }
+	AH4 FsrEasuAH(AF2 p) { AH4 res = InputTexture.GatherAlpha(samLinearClamp, p, ASU2(0, 0)); return res; }
 	AH4 FsrEasuRHO(AF2 p, ASW2 o) { AH4 res = InputTexture.GatherRed  (samLinearClamp, p, o); return res; }
 	AH4 FsrEasuGHO(AF2 p, ASW2 o) { AH4 res = InputTexture.GatherGreen(samLinearClamp, p, o); return res; }
 	AH4 FsrEasuBHO(AF2 p, ASW2 o) { AH4 res = InputTexture.GatherBlue (samLinearClamp, p, o); return res; }
-	AH3 FsrEasuSampleH(AF2 p) { AH3 res = InputTexture.SampleLevel(samLinearClamp, p, 0).xyz; return res; }
+	AH4 FsrEasuAHO(AF2 p, ASW2 o) { AH4 res = InputTexture.GatherAlpha(samLinearClamp, p, o); return res; }
+	AH4 FsrEasuSampleH(AF2 p) { AH4 res = InputTexture.SampleLevel(samLinearClamp, p, 0);     return res; }
 #else
 	#include "ffx_a.ush"
 	Texture2D InputTexture;
@@ -74,6 +76,7 @@ float2 VPColor_ViewportMin;
 	AF4 FsrEasuRF(AF2 p) { AF4 res = InputTexture.GatherRed  (samLinearClamp, p, ASU2(0, 0)); return res; }
 	AF4 FsrEasuGF(AF2 p) { AF4 res = InputTexture.GatherGreen(samLinearClamp, p, ASU2(0, 0)); return res; }
 	AF4 FsrEasuBF(AF2 p) { AF4 res = InputTexture.GatherBlue (samLinearClamp, p, ASU2(0, 0)); return res; }
+	AF4 FsrEasuAF(AF2 p) { AF4 res = InputTexture.GatherAlpha(samLinearClamp, p, ASU2(0, 0)); return res; }
 #endif
 
 // should be included after ffx_a.ush, defines FSR_OUTPUTDEVICE
@@ -98,24 +101,22 @@ AF4 FidelityFXSuperResolution(AU2 gxy)
 // FSR / EASU()
 //
 #if ENABLE_FP16
-	AH3 Gamma2Color = AH3(0,0,0);
+	AH4 Gamma2Color = AH4(0,0,0,0);
 	#if ES3_1_PROFILE	
 		FsrEasuL(Gamma2Color, gxy, Const0, Const1, Const2, Const3);
 	#else
 		FsrEasuH(Gamma2Color, gxy, Const0, Const1, Const2, Const3);
 	#endif
 #else
-	AF3 Gamma2Color = AF3(0,0,0);
+	AF4 Gamma2Color = AF4(0,0,0,0);
 	FsrEasuF(Gamma2Color, gxy, Const0, Const1, Const2, Const3);
 #endif // ENABLE_FP16
 
-	AF3 OutColor = AF3(Gamma2Color);
+#if !USE_PASSTHROUGH_ALPHA
+	Gamma2Color.a = 1;
+#endif // !USE_PASSTHROUGH_ALPHA
 
-#if ENABLE_FP16
-	return AH4(OutColor, 1);
-#else
-	return AF4(OutColor, 1);
-#endif
+	return Gamma2Color;
 }
 
 
diff --git a/Engine/Plugins/Runtime/MobileFSR/Shaders/Private/PostProcessMobileFFX_RCAS.usf b/Engine/Plugins/Runtime/MobileFSR/Shaders/Private/PostProcessMobileFFX_RCAS.usf
--- a/Engine/Plugins/Runtime/MobileFSR/Shaders/Private/PostProcessMobileFFX_RCAS.usf
+++ b/Engine/Plugins/Runtime/MobileFSR/Shaders/Private/PostProcessMobileFFX_RCAS.usf
@@ -47,6 +47,10 @@ float2 VPColor_ExtentInverse;
 #define FSR_RCAS_DENOISE 1
 #endif
 
+#if USE_PASSTHROUGH_ALPHA
+#define FSR_RCAS_PASSTHROUGH_ALPHA 1
+#endif
+
 #if ENABLE_FP16
 	#define FSR_RCAS_H   1
 	#define FSR_RCAS_HX2 1
@@ -89,26 +93,46 @@ float2 VPColor_ExtentInverse;
 // RCAS IMPLEMENTATION
 // 
 // =====================================================================================
-AF3 RCASPassCommon(AU2 gxy)
+#if ENABLE_FP16 // half precision floating point
+AH4 RCASPassCommon(AU2 gxy)
+#else
+AF4 RCASPassCommon(AU2 gxy)
+#endif
 {
 	//
 	// RCAS()
 	//
 	#if ENABLE_FP16 // half precision floating point
 
-		AH3 Gamma2Color = AH3(0, 0, 0);
-		FsrRcasH(Gamma2Color.r, Gamma2Color.g, Gamma2Color.b, gxy, Const0);
+		#if USE_PASSTHROUGH_ALPHA
+
+			AH4 Gamma2Color = AH4(0, 0, 0, 0);
+			FsrRcasH(Gamma2Color.r, Gamma2Color.g, Gamma2Color.b, Gamma2Color.a, gxy, Const0);
+	
+		#else
+	
+			AH4 Gamma2Color = AH4(0, 0, 0, 1);
+			FsrRcasH(Gamma2Color.r, Gamma2Color.g, Gamma2Color.b, gxy, Const0);
+	
+		#endif // USE_PASSTHROUGH_ALPHA
 	
 	#else // FP32 - full precision floating point
 
-		AF3 Gamma2Color = AF3(0, 0, 0);
-		FsrRcasF(Gamma2Color.r, Gamma2Color.g, Gamma2Color.b, gxy, Const0);
-
+		#if USE_PASSTHROUGH_ALPHA
+	
+			AF4 Gamma2Color = AF4(0, 0, 0, 0);
+			FsrRcasF(Gamma2Color.r, Gamma2Color.g, Gamma2Color.b, Gamma2Color.a, gxy, Const0);
+	
+		#else
+	
+			AF4 Gamma2Color = AF4(0, 0, 0, 1);
+			FsrRcasF(Gamma2Color.r, Gamma2Color.g, Gamma2Color.b, gxy, Const0);
+	
+		#endif // USE_PASSTHROUGH_ALPHA
+	
 	#endif // ENABLE_FP16
 
-	AF3 OutColor = AF3(Gamma2Color);
-
-	return OutColor;
+	return Gamma2Color;
 }
 
 
@@ -126,27 +150,27 @@ void MainCS(uint3 LocalThreadId : SV_GroupThreadID, uint3 WorkGroupId : SV_Group
 	AU2 gxy = ARmp8x8(LocalThreadId.x) + AU2(WorkGroupId.x << 4u, WorkGroupId.y << 4u);	
 	
 #if ENABLE_FP16
-	OutputTexture[gxy] = AH4(RCASPassCommon(gxy), 1);
+	OutputTexture[gxy] = RCASPassCommon(gxy);
 
 	gxy.x += 8u;
-	OutputTexture[gxy] = AH4(RCASPassCommon(gxy), 1);
+	OutputTexture[gxy] = RCASPassCommon(gxy);
 
 	gxy.y += 8u;
-	OutputTexture[gxy] = AH4(RCASPassCommon(gxy), 1);
+	OutputTexture[gxy] = RCASPassCommon(gxy);
 
 	gxy.x -= 8u;
-	OutputTexture[gxy] = AH4(RCASPassCommon(gxy), 1);
+	OutputTexture[gxy] = RCASPassCommon(gxy);
 #else
-	OutputTexture[gxy] = AF4(RCASPassCommon(gxy), 1);
+	OutputTexture[gxy] = RCASPassCommon(gxy);
 
 	gxy.x += 8u;
-	OutputTexture[gxy] = AF4(RCASPassCommon(gxy), 1);
+	OutputTexture[gxy] = RCASPassCommon(gxy);
 
 	gxy.y += 8u;
-	OutputTexture[gxy] = AF4(RCASPassCommon(gxy), 1);
+	OutputTexture[gxy] = RCASPassCommon(gxy);
 
 	gxy.x -= 8u;
-	OutputTexture[gxy] = AF4(RCASPassCommon(gxy), 1);
+	OutputTexture[gxy] = RCASPassCommon(gxy);
 #endif
 }
 #endif // COMPUTE_SHADER
@@ -155,12 +179,12 @@ void MainCS(uint3 LocalThreadId : SV_GroupThreadID, uint3 WorkGroupId : SV_Group
 void MainPS(noperspective float4 UVAndScreenPos : TEXCOORD0, float4 SvPosition : SV_POSITION, out AH4 OutColor : SV_Target0)
 {
 	AU2 gxy = AU2(SvPosition.xy);
-	OutColor = AH4(RCASPassCommon(gxy), 1);
+	OutColor = RCASPassCommon(gxy);
 }
 #else
 void MainPS(noperspective float4 UVAndScreenPos : TEXCOORD0, float4 SvPosition : SV_POSITION, out float4 OutColor : SV_Target0)
 {
 	AU2 gxy = AU2(SvPosition.xy);
-	OutColor = AF4(RCASPassCommon(gxy), 1);
+	OutColor = RCASPassCommon(gxy);
 }
-#endif
\ No newline at end of file
+#endif
diff --git a/Engine/Plugins/Runtime/MobileFSR/Shaders/Private/ffx_fsr1.ush b/Engine/Plugins/Runtime/MobileFSR/Shaders/Private/ffx_fsr1.ush
--- a/Engine/Plugins/Runtime/MobileFSR/Shaders/Private/ffx_fsr1.ush
+++ b/Engine/Plugins/Runtime/MobileFSR/Shaders/Private/ffx_fsr1.ush
@@ -182,17 +182,18 @@ A_STATIC void FsrEasuConOffset(
  //AF4 FsrEasuRF(AF2 p);
  //AF4 FsrEasuGF(AF2 p);
  //AF4 FsrEasuBF(AF2 p);
+ //AF4 FsrEasuAF(AF2 p);
 //------------------------------------------------------------------------------------------------------------------------------
  // Filtering for a given tap for the scalar.
  void FsrEasuTapF(
- inout AF3 aC, // Accumulated color, with negative lobe.
+ inout AF4 aC, // Accumulated color, with negative lobe.
  inout AF1 aW, // Accumulated weight.
  AF2 off, // Pixel offset from resolve position to tap.
  AF2 dir, // Gradient direction.
  AF2 len, // Length.
  AF1 lob, // Negative lobe strength.
  AF1 clp, // Clipping point.
- AF3 c){ // Tap color.
+ AF4 c){ // Tap color.
   // Rotate offset by direction.
   AF2 v;
   v.x=(off.x*( dir.x))+(off.y*dir.y);
@@ -261,7 +262,7 @@ A_STATIC void FsrEasuConOffset(
   len+=lenY*w;}
 //------------------------------------------------------------------------------------------------------------------------------
  void FsrEasuF(
- out AF3 pix,
+ out AF4 pix,
  AU2 ip, // Integer pixel position in output.
  AU4 con0, // Constants generated by FsrEasuCon().
  AU4 con1,
@@ -297,15 +298,19 @@ A_STATIC void FsrEasuConOffset(
   AF4 bczzR=FsrEasuRF(p0);
   AF4 bczzG=FsrEasuGF(p0);
   AF4 bczzB=FsrEasuBF(p0);
+  AF4 bczzA=FsrEasuAF(p0);
   AF4 ijfeR=FsrEasuRF(p1);
   AF4 ijfeG=FsrEasuGF(p1);
   AF4 ijfeB=FsrEasuBF(p1);
+  AF4 ijfeA=FsrEasuAF(p1);
   AF4 klhgR=FsrEasuRF(p2);
   AF4 klhgG=FsrEasuGF(p2);
   AF4 klhgB=FsrEasuBF(p2);
+  AF4 klhgA=FsrEasuAF(p2);
   AF4 zzonR=FsrEasuRF(p3);
   AF4 zzonG=FsrEasuGF(p3);
   AF4 zzonB=FsrEasuBF(p3);
+  AF4 zzonA=FsrEasuAF(p3);
 //------------------------------------------------------------------------------------------------------------------------------
   // Simplest multi-channel approximate luma possible (luma times 2, in 2 FMA/MAD).
   AF4 bczzL=bczzB*AF4_(0.5)+(bczzR*AF4_(0.5)+bczzG);
@@ -361,28 +366,28 @@ A_STATIC void FsrEasuConOffset(
   //  e f g h
   //  i j k l
   //    n o
-  AF3 min4=min(AMin3F3(AF3(ijfeR.z,ijfeG.z,ijfeB.z),AF3(klhgR.w,klhgG.w,klhgB.w),AF3(ijfeR.y,ijfeG.y,ijfeB.y)),
-               AF3(klhgR.x,klhgG.x,klhgB.x));
-  AF3 max4=max(AMax3F3(AF3(ijfeR.z,ijfeG.z,ijfeB.z),AF3(klhgR.w,klhgG.w,klhgB.w),AF3(ijfeR.y,ijfeG.y,ijfeB.y)),
-               AF3(klhgR.x,klhgG.x,klhgB.x));
+  AF4 min4=min(AMin3F4(AF4(ijfeR.z,ijfeG.z,ijfeB.z,ijfeA.z),AF4(klhgR.w,klhgG.w,klhgB.w,klhgA.w),AF4(ijfeR.y,ijfeG.y,ijfeB.y,ijfeA.y)),
+               AF4(klhgR.x,klhgG.x,klhgB.x,klhgA.x));
+  AF4 max4=max(AMax3F4(AF4(ijfeR.z,ijfeG.z,ijfeB.z,ijfeA.z),AF4(klhgR.w,klhgG.w,klhgB.w,klhgA.w),AF4(ijfeR.y,ijfeG.y,ijfeB.y,ijfeA.y)),
+               AF4(klhgR.x,klhgG.x,klhgB.x,klhgA.x));
   // Accumulation.
-  AF3 aC=AF3_(0.0);
+  AF4 aC=AF4_(0.0);
   AF1 aW=AF1_(0.0);
-  FsrEasuTapF(aC,aW,AF2( 0.0,-1.0)-pp,dir,len2,lob,clp,AF3(bczzR.x,bczzG.x,bczzB.x)); // b
-  FsrEasuTapF(aC,aW,AF2( 1.0,-1.0)-pp,dir,len2,lob,clp,AF3(bczzR.y,bczzG.y,bczzB.y)); // c
-  FsrEasuTapF(aC,aW,AF2(-1.0, 1.0)-pp,dir,len2,lob,clp,AF3(ijfeR.x,ijfeG.x,ijfeB.x)); // i
-  FsrEasuTapF(aC,aW,AF2( 0.0, 1.0)-pp,dir,len2,lob,clp,AF3(ijfeR.y,ijfeG.y,ijfeB.y)); // j
-  FsrEasuTapF(aC,aW,AF2( 0.0, 0.0)-pp,dir,len2,lob,clp,AF3(ijfeR.z,ijfeG.z,ijfeB.z)); // f
-  FsrEasuTapF(aC,aW,AF2(-1.0, 0.0)-pp,dir,len2,lob,clp,AF3(ijfeR.w,ijfeG.w,ijfeB.w)); // e
-  FsrEasuTapF(aC,aW,AF2( 1.0, 1.0)-pp,dir,len2,lob,clp,AF3(klhgR.x,klhgG.x,klhgB.x)); // k
-  FsrEasuTapF(aC,aW,AF2( 2.0, 1.0)-pp,dir,len2,lob,clp,AF3(klhgR.y,klhgG.y,klhgB.y)); // l
-  FsrEasuTapF(aC,aW,AF2( 2.0, 0.0)-pp,dir,len2,lob,clp,AF3(klhgR.z,klhgG.z,klhgB.z)); // h
-  FsrEasuTapF(aC,aW,AF2( 1.0, 0.0)-pp,dir,len2,lob,clp,AF3(klhgR.w,klhgG.w,klhgB.w)); // g
-  FsrEasuTapF(aC,aW,AF2( 1.0, 2.0)-pp,dir,len2,lob,clp,AF3(zzonR.z,zzonG.z,zzonB.z)); // o
-  FsrEasuTapF(aC,aW,AF2( 0.0, 2.0)-pp,dir,len2,lob,clp,AF3(zzonR.w,zzonG.w,zzonB.w)); // n
+  FsrEasuTapF(aC,aW,AF2( 0.0,-1.0)-pp,dir,len2,lob,clp,AF4(bczzR.x,bczzG.x,bczzB.x,bczzA.x)); // b
+  FsrEasuTapF(aC,aW,AF2( 1.0,-1.0)-pp,dir,len2,lob,clp,AF4(bczzR.y,bczzG.y,bczzB.y,bczzA.y)); // c
+  FsrEasuTapF(aC,aW,AF2(-1.0, 1.0)-pp,dir,len2,lob,clp,AF4(ijfeR.x,ijfeG.x,ijfeB.x,ijfeA.x)); // i
+  FsrEasuTapF(aC,aW,AF2( 0.0, 1.0)-pp,dir,len2,lob,clp,AF4(ijfeR.y,ijfeG.y,ijfeB.y,ijfeA.y)); // j
+  FsrEasuTapF(aC,aW,AF2( 0.0, 0.0)-pp,dir,len2,lob,clp,AF4(ijfeR.z,ijfeG.z,ijfeB.z,ijfeA.z)); // f
+  FsrEasuTapF(aC,aW,AF2(-1.0, 0.0)-pp,dir,len2,lob,clp,AF4(ijfeR.w,ijfeG.w,ijfeB.w,ijfeA.w)); // e
+  FsrEasuTapF(aC,aW,AF2( 1.0, 1.0)-pp,dir,len2,lob,clp,AF4(klhgR.x,klhgG.x,klhgB.x,klhgA.x)); // k
+  FsrEasuTapF(aC,aW,AF2( 2.0, 1.0)-pp,dir,len2,lob,clp,AF4(klhgR.y,klhgG.y,klhgB.y,klhgA.y)); // l
+  FsrEasuTapF(aC,aW,AF2( 2.0, 0.0)-pp,dir,len2,lob,clp,AF4(klhgR.z,klhgG.z,klhgB.z,klhgA.z)); // h
+  FsrEasuTapF(aC,aW,AF2( 1.0, 0.0)-pp,dir,len2,lob,clp,AF4(klhgR.w,klhgG.w,klhgB.w,klhgA.w)); // g
+  FsrEasuTapF(aC,aW,AF2( 1.0, 2.0)-pp,dir,len2,lob,clp,AF4(zzonR.z,zzonG.z,zzonB.z,zzonA.z)); // o
+  FsrEasuTapF(aC,aW,AF2( 0.0, 2.0)-pp,dir,len2,lob,clp,AF4(zzonR.w,zzonG.w,zzonB.w,zzonA.w)); // n
 //------------------------------------------------------------------------------------------------------------------------------
   // Normalize and dering.
-  pix=min(max4,max(min4,aC*AF3_(ARcpF1(aW))));}
+  pix=min(max4,max(min4,aC*AF4_(ARcpF1(aW))));}
 #endif
 ////////////////////////////////////////////////////////////////////////////////////////////////////////////////////////////////
 ////////////////////////////////////////////////////////////////////////////////////////////////////////////////////////////////
@@ -395,17 +400,18 @@ A_STATIC void FsrEasuConOffset(
  //AH4 FsrEasuRH(AF2 p);
  //AH4 FsrEasuGH(AF2 p);
  //AH4 FsrEasuBH(AF2 p);
+ //AH4 FsrEasuAH(AF2 p);
 //------------------------------------------------------------------------------------------------------------------------------
  // This runs 2 taps in parallel.
  void FsrEasuTapH(
- inout AH2 aCR,inout AH2 aCG,inout AH2 aCB,
+ inout AH2 aCR,inout AH2 aCG,inout AH2 aCB,inout AH2 aCA,
  inout AH2 aW,
  AH2 offX,AH2 offY,
  AH2 dir,
  AH2 len,
  AH1 lob,
  AH1 clp,
- AH2 cR,AH2 cG,AH2 cB){
+ AH2 cR,AH2 cG,AH2 cB,AH2 cA){
   AH2 vX,vY;
   vX=offX*  dir.xx +offY*dir.yy;
   vY=offX*(-dir.yy)+offY*dir.xx;
@@ -418,7 +424,7 @@ A_STATIC void FsrEasuConOffset(
   wA*=wA;
   wB=AH2_(25.0/16.0)*wB+AH2_(-(25.0/16.0-1.0));
   AH2 w=wB*wA;
-  aCR+=cR*w;aCG+=cG*w;aCB+=cB*w;aW+=w;}
+  aCR+=cR*w;aCG+=cG*w;aCB+=cB*w;aCA+=cA*w;aW+=w;}
 //------------------------------------------------------------------------------------------------------------------------------
  // This runs 2 taps in parallel.
  void FsrEasuSetH(
@@ -451,7 +457,7 @@ A_STATIC void FsrEasuConOffset(
   lenP+=lenY*w;}
 //------------------------------------------------------------------------------------------------------------------------------
  void FsrEasuH(
- out AH3 pix,
+ out AH4 pix,
  AU2 ip,
  AU4 con0,
  AU4 con1,
@@ -470,15 +476,19 @@ A_STATIC void FsrEasuConOffset(
   AH4 bczzR=FsrEasuRH(p0);
   AH4 bczzG=FsrEasuGH(p0);
   AH4 bczzB=FsrEasuBH(p0);
+  AH4 bczzA=FsrEasuAH(p0);
   AH4 ijfeR=FsrEasuRH(p1);
   AH4 ijfeG=FsrEasuGH(p1);
   AH4 ijfeB=FsrEasuBH(p1);
+  AH4 ijfeA=FsrEasuAH(p1);
   AH4 klhgR=FsrEasuRH(p2);
   AH4 klhgG=FsrEasuGH(p2);
   AH4 klhgB=FsrEasuBH(p2);
+  AH4 klhgA=FsrEasuAH(p2);
   AH4 zzonR=FsrEasuRH(p3);
   AH4 zzonG=FsrEasuGH(p3);
   AH4 zzonB=FsrEasuBH(p3);
+  AH4 zzonA=FsrEasuAH(p3);
 //------------------------------------------------------------------------------------------------------------------------------
   AH4 bczzL=bczzB*AH4_(0.5)+(bczzR*AH4_(0.5)+bczzG);
   AH4 ijfeL=ijfeB*AH4_(0.5)+(ijfeR*AH4_(0.5)+ijfeG);
@@ -523,22 +533,24 @@ A_STATIC void FsrEasuConOffset(
   AH2 bothR=max(max(AH2(-ijfeR.z,ijfeR.z),AH2(-klhgR.w,klhgR.w)),max(AH2(-ijfeR.y,ijfeR.y),AH2(-klhgR.x,klhgR.x)));
   AH2 bothG=max(max(AH2(-ijfeG.z,ijfeG.z),AH2(-klhgG.w,klhgG.w)),max(AH2(-ijfeG.y,ijfeG.y),AH2(-klhgG.x,klhgG.x)));
   AH2 bothB=max(max(AH2(-ijfeB.z,ijfeB.z),AH2(-klhgB.w,klhgB.w)),max(AH2(-ijfeB.y,ijfeB.y),AH2(-klhgB.x,klhgB.x)));
+  AH2 bothA=max(max(AH2(-ijfeA.z,ijfeA.z),AH2(-klhgA.w,klhgA.w)),max(AH2(-ijfeA.y,ijfeA.y),AH2(-klhgA.x,klhgA.x)));
   // This part is different for FP16, working pairs of taps at a time.
   AH2 pR=AH2_(0.0);
   AH2 pG=AH2_(0.0);
   AH2 pB=AH2_(0.0);
+  AH2 pA=AH2_(0.0);
   AH2 pW=AH2_(0.0);
-  FsrEasuTapH(pR,pG,pB,pW,AH2( 0.0, 1.0)-ppp.xx,AH2(-1.0,-1.0)-ppp.yy,dir,len2,lob,clp,bczzR.xy,bczzG.xy,bczzB.xy);
-  FsrEasuTapH(pR,pG,pB,pW,AH2(-1.0, 0.0)-ppp.xx,AH2( 1.0, 1.0)-ppp.yy,dir,len2,lob,clp,ijfeR.xy,ijfeG.xy,ijfeB.xy);
-  FsrEasuTapH(pR,pG,pB,pW,AH2( 0.0,-1.0)-ppp.xx,AH2( 0.0, 0.0)-ppp.yy,dir,len2,lob,clp,ijfeR.zw,ijfeG.zw,ijfeB.zw);
-  FsrEasuTapH(pR,pG,pB,pW,AH2( 1.0, 2.0)-ppp.xx,AH2( 1.0, 1.0)-ppp.yy,dir,len2,lob,clp,klhgR.xy,klhgG.xy,klhgB.xy);
-  FsrEasuTapH(pR,pG,pB,pW,AH2( 2.0, 1.0)-ppp.xx,AH2( 0.0, 0.0)-ppp.yy,dir,len2,lob,clp,klhgR.zw,klhgG.zw,klhgB.zw);
-  FsrEasuTapH(pR,pG,pB,pW,AH2( 1.0, 0.0)-ppp.xx,AH2( 2.0, 2.0)-ppp.yy,dir,len2,lob,clp,zzonR.zw,zzonG.zw,zzonB.zw);
-  AH3 aC=AH3(pR.x+pR.y,pG.x+pG.y,pB.x+pB.y);
+  FsrEasuTapH(pR,pG,pB,pA,pW,AH2( 0.0, 1.0)-ppp.xx,AH2(-1.0,-1.0)-ppp.yy,dir,len2,lob,clp,bczzR.xy,bczzG.xy,bczzB.xy,bczzA.xy);
+  FsrEasuTapH(pR,pG,pB,pA,pW,AH2(-1.0, 0.0)-ppp.xx,AH2( 1.0, 1.0)-ppp.yy,dir,len2,lob,clp,ijfeR.xy,ijfeG.xy,ijfeB.xy,ijfeA.xy);
+  FsrEasuTapH(pR,pG,pB,pA,pW,AH2( 0.0,-1.0)-ppp.xx,AH2( 0.0, 0.0)-ppp.yy,dir,len2,lob,clp,ijfeR.zw,ijfeG.zw,ijfeB.zw,ijfeA.zw);
+  FsrEasuTapH(pR,pG,pB,pA,pW,AH2( 1.0, 2.0)-ppp.xx,AH2( 1.0, 1.0)-ppp.yy,dir,len2,lob,clp,klhgR.xy,klhgG.xy,klhgB.xy,klhgA.xy);
+  FsrEasuTapH(pR,pG,pB,pA,pW,AH2( 2.0, 1.0)-ppp.xx,AH2( 0.0, 0.0)-ppp.yy,dir,len2,lob,clp,klhgR.zw,klhgG.zw,klhgB.zw,klhgA.zw);
+  FsrEasuTapH(pR,pG,pB,pA,pW,AH2( 1.0, 0.0)-ppp.xx,AH2( 2.0, 2.0)-ppp.yy,dir,len2,lob,clp,zzonR.zw,zzonG.zw,zzonB.zw,zzonA.zw);
+  AH4 aC=AH4(pR.x+pR.y,pG.x+pG.y,pB.x+pB.y,pA.x+pA.y);
   AH1 aW=pW.x+pW.y;
 //------------------------------------------------------------------------------------------------------------------------------
   // Slightly different for FP16 version due to combined min and max.
-  pix=min(AH3(bothR.y,bothG.y,bothB.y),max(-AH3(bothR.x,bothG.x,bothB.x),aC*AH3_(ARcpH1(aW))));}
+  pix=min(AH4(bothR.y,bothG.y,bothB.y,bothA.y),max(-AH4(bothR.x,bothG.x,bothB.x,bothA.x),aC*AH4_(ARcpH1(aW))));}
 //==============================================================================================================================
 // An optimized AMD FSR's EASU implementation for Mobiles
 // Based on [REDACTED_PROFILE_URL]/FidelityFX-FSR/blob/master/ffx-fsr/ffx_fsr1.h
@@ -565,10 +577,10 @@ A_STATIC void FsrEasuConOffset(
 // Copyright (c) 2022, Qualcomm Innovation Center, Inc. All rights reserved.
 // 
 // -- FsrEasuSampleH should be implemented by calling shader, like following:
-//    AH3 FsrEasuSampleH(AF2 p) { return MyTex.SampleLevel(LinearSampler, p, 0).xyz; }
+//    AH4 FsrEasuSampleH(AF2 p) { return MyTex.SampleLevel(LinearSampler, p, 0); }
 //==============================================================================================================================
  void FsrEasuL(
- out AH3 pix,
+ out AH4 pix,
  AU2 ip,
  AU4 con0,
  AU4 con1,
@@ -581,11 +593,11 @@ A_STATIC void FsrEasuConOffset(
   //    E
   AF2 pp=AF2(ip)*AF2_AU2(con0.xy)+AF2_AU2(con0.zw);
   AF2 tc=(pp+AF2_(0.5))*AF2_AU2(con1.xy);
-  AH3 sA=FsrEasuSampleH(tc-AF2(0, AF1_AU1(con1.y)));
-  AH3 sB=FsrEasuSampleH(tc-AF2(AF1_AU1(con1.x), 0));
-  AH3 sC=FsrEasuSampleH(tc);
-  AH3 sD=FsrEasuSampleH(tc+AF2(AF1_AU1(con1.x), 0));
-  AH3 sE=FsrEasuSampleH(tc+AF2(0, AF1_AU1(con1.y)));
+  AH4 sA=FsrEasuSampleH(tc-AF2(0, AF1_AU1(con1.y)));
+  AH4 sB=FsrEasuSampleH(tc-AF2(AF1_AU1(con1.x), 0));
+  AH4 sC=FsrEasuSampleH(tc);
+  AH4 sD=FsrEasuSampleH(tc+AF2(AF1_AU1(con1.x), 0));
+  AH4 sE=FsrEasuSampleH(tc+AF2(0, AF1_AU1(con1.y)));
   AH1 lA=sA.r*AH1_(0.5)+sA.g;
   AH1 lB=sB.r*AH1_(0.5)+sB.g;
   AH1 lC=sC.r*AH1_(0.5)+sC.g;
@@ -637,38 +649,43 @@ A_STATIC void FsrEasuConOffset(
   AH4 fgcbR=FsrEasuRH(p0);
   AH4 fgcbG=FsrEasuGH(p0);
   AH4 fgcbB=FsrEasuBH(p0);
+  AH4 fgcbA=FsrEasuAH(p0);
   AH4 ijfeR=FsrEasuRH(p1);
   AH4 ijfeG=FsrEasuGH(p1);
   AH4 ijfeB=FsrEasuBH(p1);
+  AH4 ijfeA=FsrEasuAH(p1);
   AH4 klhgR=FsrEasuRH(p2);
   AH4 klhgG=FsrEasuGH(p2);
   AH4 klhgB=FsrEasuBH(p2);
+  AH4 klhgA=FsrEasuAH(p2);
   AH4 nokjR=FsrEasuRH(p3);
   AH4 nokjG=FsrEasuGH(p3);
   AH4 nokjB=FsrEasuBH(p3);
+  AH4 nokjA=FsrEasuAH(p3);
 
 //------------------------------------------------------------------------------------------------------------------------------
   // This part is different for FP16, working pairs of taps at a time.
   AH2 pR=AH2_(0.0);
   AH2 pG=AH2_(0.0);
   AH2 pB=AH2_(0.0);
+  AH2 pA=AH2_(0.0);
   AH2 pW=AH2_(0.0);
-  FsrEasuTapH(pR,pG,pB,pW,AH2( 1.0, 0.0)-ppp.xx,AH2(-1.0,-1.0)-ppp.yy,dir,len2,lob,clp,fgcbR.zw,fgcbG.zw,fgcbB.zw);
-  FsrEasuTapH(pR,pG,pB,pW,AH2(-1.0, 0.0)-ppp.xx,AH2( 1.0, 1.0)-ppp.yy,dir,len2,lob,clp,ijfeR.xy,ijfeG.xy,ijfeB.xy);
-  FsrEasuTapH(pR,pG,pB,pW,AH2( 0.0,-1.0)-ppp.xx,AH2( 0.0, 0.0)-ppp.yy,dir,len2,lob,clp,ijfeR.zw,ijfeG.zw,ijfeB.zw);
-  FsrEasuTapH(pR,pG,pB,pW,AH2( 1.0, 2.0)-ppp.xx,AH2( 1.0, 1.0)-ppp.yy,dir,len2,lob,clp,klhgR.xy,klhgG.xy,klhgB.xy);
-  FsrEasuTapH(pR,pG,pB,pW,AH2( 2.0, 1.0)-ppp.xx,AH2( 0.0, 0.0)-ppp.yy,dir,len2,lob,clp,klhgR.zw,klhgG.zw,klhgB.zw);
-  FsrEasuTapH(pR,pG,pB,pW,AH2( 0.0, 1.0)-ppp.xx,AH2( 2.0, 2.0)-ppp.yy,dir,len2,lob,clp,nokjR.xy,nokjG.xy,nokjB.xy);
-  AH3 aC=AH3(pR.x+pR.y,pG.x+pG.y,pB.x+pB.y);
+  FsrEasuTapH(pR,pG,pB,pA,pW,AH2( 1.0, 0.0)-ppp.xx,AH2(-1.0,-1.0)-ppp.yy,dir,len2,lob,clp,fgcbR.zw,fgcbG.zw,fgcbB.zw,fgcbA.zw);
+  FsrEasuTapH(pR,pG,pB,pA,pW,AH2(-1.0, 0.0)-ppp.xx,AH2( 1.0, 1.0)-ppp.yy,dir,len2,lob,clp,ijfeR.xy,ijfeG.xy,ijfeB.xy,ijfeA.xy);
+  FsrEasuTapH(pR,pG,pB,pA,pW,AH2( 0.0,-1.0)-ppp.xx,AH2( 0.0, 0.0)-ppp.yy,dir,len2,lob,clp,ijfeR.zw,ijfeG.zw,ijfeB.zw,ijfeA.zw);
+  FsrEasuTapH(pR,pG,pB,pA,pW,AH2( 1.0, 2.0)-ppp.xx,AH2( 1.0, 1.0)-ppp.yy,dir,len2,lob,clp,klhgR.xy,klhgG.xy,klhgB.xy,klhgA.xy);
+  FsrEasuTapH(pR,pG,pB,pA,pW,AH2( 2.0, 1.0)-ppp.xx,AH2( 0.0, 0.0)-ppp.yy,dir,len2,lob,clp,klhgR.zw,klhgG.zw,klhgB.zw,klhgA.zw);
+  FsrEasuTapH(pR,pG,pB,pA,pW,AH2( 0.0, 1.0)-ppp.xx,AH2( 2.0, 2.0)-ppp.yy,dir,len2,lob,clp,nokjR.xy,nokjG.xy,nokjB.xy,nokjA.xy);
+  AH4 aC=AH4(pR.x+pR.y,pG.x+pG.y,pB.x+pB.y,pA.x+pA.y);
   AH1 aW=pW.x+pW.y;
 //------------------------------------------------------------------------------------------------------------------------------
   // Qualcomm - normalize and de-ring.  While the Atyuwen original blog post did not see any visual changes from removing the de-ringer
   //            we saw noticable artifacts on high contrast sharp edges.
-  AH3 min4 = min(AMin3H3(AH3(ijfeR.z, ijfeG.z, ijfeB.z), AH3(klhgR.w, klhgG.w, klhgB.w), AH3(ijfeR.y, ijfeG.y, ijfeB.y)),
-      AH3(klhgR.x, klhgG.x, klhgB.x));
-  AH3 max4 = max(AMax3H3(AH3(ijfeR.z, ijfeG.z, ijfeB.z), AH3(klhgR.w, klhgG.w, klhgB.w), AH3(ijfeR.y, ijfeG.y, ijfeB.y)),
-      AH3(klhgR.x, klhgG.x, klhgB.x));
-  pix = min(max4, max(min4, aC * AH3_(ARcpH1(aW))));}
+  AH4 min4 = min(AMin3H4(AH4(ijfeR.z, ijfeG.z, ijfeB.z, ijfeB.a), AH4(klhgR.w, klhgG.w, klhgB.w, klhgA.w), AH4(ijfeR.y, ijfeG.y, ijfeB.y, ijfeA.y)),
+      AH4(klhgR.x, klhgG.x, klhgB.x, klhgA.x));
+  AH4 max4 = max(AMax3H4(AH4(ijfeR.z, ijfeG.z, ijfeB.z, ijfeA.z), AH4(klhgR.w, klhgG.w, klhgB.w, klhgA.w), AH4(ijfeR.y, ijfeG.y, ijfeB.y, ijfeA.y)),
+      AH4(klhgR.x, klhgG.x, klhgB.x, klhgA.x));
+  pix = min(max4, max(min4, aC * AH4_(ARcpH1(aW))));}
 #endif
 ////////////////////////////////////////////////////////////////////////////////////////////////////////////////////////////////
 ////////////////////////////////////////////////////////////////////////////////////////////////////////////////////////////////
diff --git a/Engine/Plugins/Runtime/MobileFSR/Source/MobileFSR/Private/MobileFSRUpscaler.cpp b/Engine/Plugins/Runtime/MobileFSR/Source/MobileFSR/Private/MobileFSRUpscaler.cpp
--- a/Engine/Plugins/Runtime/MobileFSR/Source/MobileFSR/Private/MobileFSRUpscaler.cpp
+++ b/Engine/Plugins/Runtime/MobileFSR/Source/MobileFSR/Private/MobileFSRUpscaler.cpp
@@ -4,6 +4,12 @@
 #include "DataDrivenShaderPlatformInfo.h"
 #include "ScenePrivate.h"
 
+bool IsPostProcessingWithAlphaChannelSupported()
+{
+	static IConsoleVariable* CVarPropagateAlpha = IConsoleManager::Get().FindConsoleVariable(TEXT("r.PostProcessing.PropagateAlpha"));
+	return CVarPropagateAlpha ? CVarPropagateAlpha->GetBool() : false;
+}
+
 #define A_CPU 1
 // ffx_a.ush will override these definitions
 #ifdef A_STATIC
@@ -22,6 +28,7 @@
 
 // permutation domains
 class FMobileFSR_UseFP16Dim : SHADER_PERMUTATION_BOOL("ENABLE_FP16");
+class FMobileFSR_UsePassthroughAlpha : SHADER_PERMUTATION_BOOL("USE_PASSTHROUGH_ALPHA");
 class FMobileRCAS_DenoiseDim : SHADER_PERMUTATION_BOOL("USE_RCAS_DENOISE");
 
 BEGIN_SHADER_PARAMETER_STRUCT(FMobileFSRPassParameters, )
@@ -71,7 +78,7 @@ public:
 	DECLARE_GLOBAL_SHADER(FMobileFSRCS);
 	SHADER_USE_PARAMETER_STRUCT(FMobileFSRCS, FGlobalShader);
 
-	using FPermutationDomain = TShaderPermutationDomain<FMobileFSR_UseFP16Dim>;
+	using FPermutationDomain = TShaderPermutationDomain<FMobileFSR_UseFP16Dim, FMobileFSR_UsePassthroughAlpha>;
 
 	BEGIN_SHADER_PARAMETER_STRUCT(FParameters, )
 		SHADER_PARAMETER_STRUCT_INCLUDE(FMobileFSRPassParameters, FSR)
@@ -113,7 +120,7 @@ public:
 	DECLARE_GLOBAL_SHADER(FMobileFSRPS);
 	SHADER_USE_PARAMETER_STRUCT(FMobileFSRPS, FGlobalShader);
 
-	using FPermutationDomain = TShaderPermutationDomain<FMobileFSR_UseFP16Dim>;
+	using FPermutationDomain = TShaderPermutationDomain<FMobileFSR_UseFP16Dim, FMobileFSR_UsePassthroughAlpha>;
 
 	BEGIN_SHADER_PARAMETER_STRUCT(FParameters, )
 		SHADER_PARAMETER_STRUCT_INCLUDE(FMobileFSRPassParameters, FSR)
@@ -204,6 +211,7 @@ FScreenPassTexture AddEASUPass(FRDGBuilder& GraphBuilder, const FViewInfo& View,
 
 		FMobileFSRPS::FPermutationDomain PSPermutationVector;
 		PSPermutationVector.Set<FMobileFSR_UseFP16Dim>(bUseFP16);
+		PSPermutationVector.Set<FMobileFSR_UsePassthroughAlpha>(IsPostProcessingWithAlphaChannelSupported());
 
 		TShaderMapRef<FMobileFSRPS> PixelShader(View.ShaderMap, PSPermutationVector);
 
@@ -236,6 +244,7 @@ FScreenPassTexture AddEASUPass(FRDGBuilder& GraphBuilder, const FViewInfo& View,
 
 		FMobileFSRCS::FPermutationDomain CSPermutationVector;
 		CSPermutationVector.Set<FMobileFSR_UseFP16Dim>(bUseFP16);
+		CSPermutationVector.Set<FMobileFSR_UsePassthroughAlpha>(IsPostProcessingWithAlphaChannelSupported());
 
 		TShaderMapRef<FMobileFSRCS> ComputeShaderFSR(View.ShaderMap, CSPermutationVector);
 		FComputeShaderUtils::AddPass(
@@ -260,7 +269,7 @@ public:
 	DECLARE_GLOBAL_SHADER(FMobileRCASCS);
 	SHADER_USE_PARAMETER_STRUCT(FMobileRCASCS, FGlobalShader);
 
-	using FPermutationDomain = TShaderPermutationDomain<FMobileFSR_UseFP16Dim, FMobileRCAS_DenoiseDim>;
+	using FPermutationDomain = TShaderPermutationDomain<FMobileFSR_UseFP16Dim, FMobileFSR_UsePassthroughAlpha, FMobileRCAS_DenoiseDim>;
 
 	BEGIN_SHADER_PARAMETER_STRUCT(FParameters, )
 		SHADER_PARAMETER_STRUCT_INCLUDE(FMobileRCASPassParameters, RCAS)
@@ -306,7 +315,7 @@ public:
 	DECLARE_GLOBAL_SHADER(FMobileRCASPS);
 	SHADER_USE_PARAMETER_STRUCT(FMobileRCASPS, FGlobalShader);
 
-	using FPermutationDomain = TShaderPermutationDomain<FMobileFSR_UseFP16Dim, FMobileRCAS_DenoiseDim>;
+	using FPermutationDomain = TShaderPermutationDomain<FMobileFSR_UseFP16Dim, FMobileFSR_UsePassthroughAlpha, FMobileRCAS_DenoiseDim>;
 
 	BEGIN_SHADER_PARAMETER_STRUCT(FParameters, )
 		SHADER_PARAMETER_STRUCT_INCLUDE(FMobileRCASPassParameters, RCAS)
@@ -372,6 +381,7 @@ FScreenPassTexture AddCASPass(FRDGBuilder& GraphBuilder, const FViewInfo& View,
 		// grab shaders
 		FMobileRCASPS::FPermutationDomain PSPermutationVector;
 		PSPermutationVector.Set<FMobileFSR_UseFP16Dim>(bUseFP16);
+		PSPermutationVector.Set<FMobileFSR_UsePassthroughAlpha>(IsPostProcessingWithAlphaChannelSupported());
 		PSPermutationVector.Set<FMobileRCAS_DenoiseDim>(GMobileFSRCASDenoise > 0);
 
 		TShaderMapRef<FMobileRCASPS> PixelShader(View.ShaderMap, PSPermutationVector);
@@ -399,6 +409,7 @@ FScreenPassTexture AddCASPass(FRDGBuilder& GraphBuilder, const FViewInfo& View,
 
 		FMobileRCASCS::FPermutationDomain CSPermutationVector;
 		CSPermutationVector.Set<FMobileFSR_UseFP16Dim>(bUseFP16);
+		CSPermutationVector.Set<FMobileFSR_UsePassthroughAlpha>(IsPostProcessingWithAlphaChannelSupported());
 		CSPermutationVector.Set<FMobileRCAS_DenoiseDim>(GMobileFSRCASDenoise > 0);
 
 		TShaderMapRef<FMobileRCASCS> ComputeShaderRCASPass(View.ShaderMap, CSPermutationVector);
@@ -441,5 +452,4 @@ FScreenPassTexture FMobileFSRUpscaler::AddPasses(FRDGBuilder& GraphBuilder, cons
 	}
 
 	return MoveTemp(SceneColor);
-
 }
