Fix duplicate symbol

diff --git a/Engine/Plugins/Runtime/MobileFSR/Source/MobileFSR/Private/MobileFSRUpscaler.cpp b/Engine/Plugins/Runtime/MobileFSR/Source/MobileFSR/Private/MobileFSRUpscaler.cpp
--- a/Engine/Plugins/Runtime/MobileFSR/Source/MobileFSR/Private/MobileFSRUpscaler.cpp
+++ b/Engine/Plugins/Runtime/MobileFSR/Source/MobileFSR/Private/MobileFSRUpscaler.cpp
@@ -4,10 +4,13 @@
 #include "DataDrivenShaderPlatformInfo.h"
 #include "ScenePrivate.h"
 
-bool IsPostProcessingWithAlphaChannelSupported()
+namespace MobileFSRPrivate
 {
-	static IConsoleVariable* CVarPropagateAlpha = IConsoleManager::Get().FindConsoleVariable(TEXT("r.PostProcessing.PropagateAlpha"));
-	return CVarPropagateAlpha ? CVarPropagateAlpha->GetBool() : false;
+	bool IsPostProcessingWithAlphaChannelSupported()
+	{
+		static IConsoleVariable* CVarPropagateAlpha = IConsoleManager::Get().FindConsoleVariable(TEXT("r.PostProcessing.PropagateAlpha"));
+		return CVarPropagateAlpha ? CVarPropagateAlpha->GetBool() : false;
+	}
 }
 
 #define A_CPU 1
@@ -151,6 +154,8 @@ IMPLEMENT_GLOBAL_SHADER(FMobileFSRPS, "/Plugin/MobileFSR/Private/PostProcessMobi
 
 FScreenPassTexture AddEASUPass(FRDGBuilder& GraphBuilder, const FViewInfo& View, const ISpatialUpscaler::FInputs& PassInputs)
 {
+	using namespace MobileFSRPrivate;
+
 	FScreenPassTexture OutputTexture;
 	bool bUseFP16 = !IsSimulatedPlatform(View.GetShaderPlatform()) && IsDxcEnabledForPlatform(View.GetShaderPlatform());
 
@@ -345,6 +350,8 @@ IMPLEMENT_GLOBAL_SHADER(FMobileRCASPS, "/Plugin/MobileFSR/Private/PostProcessMob
 
 FScreenPassTexture AddCASPass(FRDGBuilder& GraphBuilder, const FViewInfo& View, const ISpatialUpscaler::FInputs& PassInputs)
 {
+	using namespace MobileFSRPrivate;
+
 	FScreenPassTexture OutputTexture;
 	bool bUseFP16 = !IsSimulatedPlatform(View.GetShaderPlatform()) && IsDxcEnabledForPlatform(View.GetShaderPlatform());
 
