Fixed usage of DECLARE_GPU_STAT macro (when RHI_NEW_GPU_PROFILER=1).

diff --git a/Engine/Plugins/Runtime/MobileFSR/Source/MobileFSR/Private/MobileFSRUpscaler.cpp b/Engine/Plugins/Runtime/MobileFSR/Source/MobileFSR/Private/MobileFSRUpscaler.cpp
--- a/Engine/Plugins/Runtime/MobileFSR/Source/MobileFSR/Private/MobileFSRUpscaler.cpp
+++ b/Engine/Plugins/Runtime/MobileFSR/Source/MobileFSR/Private/MobileFSRUpscaler.cpp
@@ -431,7 +431,7 @@ FScreenPassTexture AddCASPass(FRDGBuilder& GraphBuilder, const FViewInfo& View,
 	return MoveTemp(OutputTexture);
 }
 
-DECLARE_GPU_STAT(MobileFSRUpscaler)
+DECLARE_GPU_STAT(MobileFSRUpscaler);
 
 FMobileFSRUpscaler::FMobileFSRUpscaler(bool bInIsEASUPass) : bIsEASUPass(bInIsEASUPass)
 {
